Added support for choosing which platforms your project will target

TTP# 332489 - TOOLS FEATURE: Editor: Allow user to designate which platforms a project is designed for; warn user when deploying to platforms that will result in a bad time

There is now a "Target Platforms" tab in the project settings which allows you to choose which platforms your project will target. This information is stored inside the .uproject file.

If you try and launch, cook, or package for a project that isn't on the supported list, then you'll see a suppressible warning notifying you that the project may not run as expected. This is also conveyed to you via a warning icon next to platforms which aren't set as a target.

Additionally the target platform icons are shown in the tooltip on the "Open Project" dialog, as well as in the tab area of the level editor.

diff --git a/Engine/Source/Editor/EditorSettingsViewer/Private/EditorSettingsViewerModule.cpp b/Engine/Source/Editor/EditorSettingsViewer/Private/EditorSettingsViewerModule.cpp
--- a/Engine/Source/Editor/EditorSettingsViewer/Private/EditorSettingsViewerModule.cpp
+++ b/Engine/Source/Editor/EditorSettingsViewer/Private/EditorSettingsViewerModule.cpp
@@ -97,8 +97,6 @@ protected:
 		InputBindingDelegates.ResetDefaultsDelegate = FOnSettingsSectionResetDefaults::CreateRaw(this, &FEditorSettingsViewerModule::HandleInputBindingsResetToDefault);
 		InputBindingDelegates.SaveDelegate = FOnSettingsSectionSave::CreateRaw(this, &FEditorSettingsViewerModule::HandleInputBindingsSave);
 
-		InputBindingEditorPanel = FModuleManager::LoadModuleChecked<IInputBindingEditorModule>("InputBindingEditor").CreateInputBindingEditorPanel();
-
 		FSettingsSectionDelegates RegionAndLanguageDelegates;
 		RegionAndLanguageDelegates.ExportDelegate = FOnSettingsSectionExport::CreateRaw(this, &FEditorSettingsViewerModule::HandleRegionAndLanguageExport);
 		RegionAndLanguageDelegates.ImportDelegate = FOnSettingsSectionImport::CreateRaw(this, &FEditorSettingsViewerModule::HandleRegionAndLanguageImport);
@@ -112,6 +110,7 @@ protected:
 			RegionAndLanguageDelegates
 		);
 
+		TWeakPtr<SWidget> InputBindingEditorPanel = FModuleManager::LoadModuleChecked<IInputBindingEditorModule>("InputBindingEditor").CreateInputBindingEditorPanel();
 		SettingsModule.RegisterSettings("Editor", "General", "InputBindings",
 			LOCTEXT("InputBindingsSettingsName", "Keyboard Shortcuts"),
 			LOCTEXT("InputBindingsSettingsDescription", "Configure keyboard shortcuts to quickly invoke operations."),
@@ -379,9 +378,6 @@ private:
 
 private:
 
-	// Holds a pointer to the input binding editor.
-	TWeakPtr<SWidget> InputBindingEditorPanel;
-
 	// Holds a pointer to the settings editor's view model.
 	TWeakPtr<ISettingsEditorModel> SettingsEditorModelPtr;
 };
diff --git a/Engine/Source/Editor/EditorStyle/Private/SlateEditorStyle.cpp b/Engine/Source/Editor/EditorStyle/Private/SlateEditorStyle.cpp
--- a/Engine/Source/Editor/EditorStyle/Private/SlateEditorStyle.cpp
+++ b/Engine/Source/Editor/EditorStyle/Private/SlateEditorStyle.cpp
@@ -5007,6 +5007,8 @@ void FSlateEditorStyle::FStyle::SetupAutomationStyles()
 		Set( "LauncherCommand.CreateBuild", new IMAGE_BRUSH( "Launcher/Launcher_Build", Icon40x40 ) );
 		Set( "LauncherCommand.AdvancedBuild", new IMAGE_BRUSH( "Launcher/Launcher_Advanced", Icon40x40 ) );
 
+		Set( "Launcher.Platform.Warning", new IMAGE_BRUSH( "Icons/alert", Icon24x24) );
+
 #if (WITH_EDITOR || (IS_PROGRAM && PLATFORM_DESKTOP))
 		const TArray<ITargetPlatform*>& TargetPlatforms = GetTargetPlatformManager()->GetTargetPlatforms();
 		for (int32 PlatformIdx = 0; PlatformIdx < TargetPlatforms.Num(); PlatformIdx++)
diff --git a/Engine/Source/Editor/GameProjectGeneration/Private/GameProjectGenerationModule.cpp b/Engine/Source/Editor/GameProjectGeneration/Private/GameProjectGenerationModule.cpp
--- a/Engine/Source/Editor/GameProjectGeneration/Private/GameProjectGenerationModule.cpp
+++ b/Engine/Source/Editor/GameProjectGeneration/Private/GameProjectGenerationModule.cpp
@@ -77,9 +77,16 @@ bool FGameProjectGenerationModule::UpdateCodeResourceFiles(TArray<FString>& OutC
 	return GameProjectUtils::GenerateGameResourceFiles(GameModuleSourcePath, FApp::GetGameName(), OutCreatedFiles, OutFailReason);
 }
 
+
 void FGameProjectGenerationModule::CheckAndWarnProjectFilenameValid()
 {
 	GameProjectUtils::CheckAndWarnProjectFilenameValid();
 }
 
+
+void FGameProjectGenerationModule::UpdateSupportedTargetPlatforms(const FName& InPlatformName, const bool bIsSupported)
+{
+	GameProjectUtils::UpdateSupportedTargetPlatforms(InPlatformName, bIsSupported);
+}
+
 #undef LOCTEXT_NAMESPACE
diff --git a/Engine/Source/Editor/GameProjectGeneration/Private/GameProjectUtils.cpp b/Engine/Source/Editor/GameProjectGeneration/Private/GameProjectUtils.cpp
--- a/Engine/Source/Editor/GameProjectGeneration/Private/GameProjectUtils.cpp
+++ b/Engine/Source/Editor/GameProjectGeneration/Private/GameProjectUtils.cpp
@@ -623,11 +623,10 @@ bool GameProjectUtils::BuildGameBinaries(const FString& ProjectFilename, FText&
 
 void GameProjectUtils::CheckForOutOfDateGameProjectFile()
 {
-	const FString& LoadedProjectFilePath = FPaths::IsProjectFilePathSet() ? FPaths::GetProjectFilePath() : FString();
-	if ( !LoadedProjectFilePath.IsEmpty() )
+	if ( FPaths::IsProjectFilePathSet() )
 	{
 		FProjectStatus ProjectStatus;
-		if (IProjectManager::Get().QueryStatusForProject(LoadedProjectFilePath, ProjectStatus))
+		if (IProjectManager::Get().QueryStatusForCurrentProject(ProjectStatus))
 		{
 			if ( ProjectStatus.bRequiresUpdate )
 			{
@@ -1991,6 +1990,28 @@ bool GameProjectUtils::DuplicateProjectForUpgrade( const FString& InProjectFile,
 	return true;
 }
 
+void GameProjectUtils::UpdateSupportedTargetPlatforms(const FName& InPlatformName, const bool bIsSupported)
+{
+	const FString& ProjectFilename = FPaths::GetProjectFilePath();
+	if(!ProjectFilename.IsEmpty())
+	{
+		// First attempt to check out the file if SCC is enabled
+		if(ISourceControlModule::Get().IsEnabled())
+		{
+			FText UnusedFailReason;
+			CheckoutGameProjectFile(ProjectFilename, UnusedFailReason);
+		}
+
+		// Second make sure the file is writable
+		if(FPlatformFileManager::Get().GetPlatformFile().IsReadOnly(*ProjectFilename))
+		{
+			FPlatformFileManager::Get().GetPlatformFile().SetReadOnly(*ProjectFilename, false);
+		}
+
+		IProjectManager::Get().UpdateSupportedTargetPlatformsForCurrentProject(InPlatformName, bIsSupported);
+	}
+}
+
 bool GameProjectUtils::ReadTemplateFile(const FString& TemplateFileName, FString& OutFileContents, FText& OutFailReason)
 {
 	const FString FullFileName = FPaths::EngineContentDir() / TEXT("Editor") / TEXT("Templates") / TemplateFileName;
diff --git a/Engine/Source/Editor/GameProjectGeneration/Private/GameProjectUtils.h b/Engine/Source/Editor/GameProjectGeneration/Private/GameProjectUtils.h
--- a/Engine/Source/Editor/GameProjectGeneration/Private/GameProjectUtils.h
+++ b/Engine/Source/Editor/GameProjectGeneration/Private/GameProjectUtils.h
@@ -231,6 +231,15 @@ public:
 	/** Creates a copy of a project directory in order to upgrade it. */
 	static bool DuplicateProjectForUpgrade( const FString& InProjectFile, FString &OutNewProjectFile );
 
+	/**
+	 * Update the list of supported target platforms based upon the parameters provided
+	 * This will take care of checking out and saving the updated .uproject file automatically
+	 * 
+	 * @param	InPlatformName		Name of the platform to target (eg, WindowsNoEditor)
+	 * @param	bIsSupported		true if the platform should be supported by this project, false if it should not
+	 */
+	static void UpdateSupportedTargetPlatforms(const FName& InPlatformName, const bool bIsSupported);
+
 private:
 	/** Generates a new project without using a template project */
 	static bool GenerateProjectFromScratch(const FString& NewProjectFile, bool bShouldGenerateCode, bool bCopyStarterContent, FText& OutFailReason);
diff --git a/Engine/Source/Editor/GameProjectGeneration/Private/SProjectBrowser.cpp b/Engine/Source/Editor/GameProjectGeneration/Private/SProjectBrowser.cpp
--- a/Engine/Source/Editor/GameProjectGeneration/Private/SProjectBrowser.cpp
+++ b/Engine/Source/Editor/GameProjectGeneration/Private/SProjectBrowser.cpp
@@ -7,6 +7,7 @@
 #include "SVerbChoiceDialog.h"
 #include "UProjectInfo.h"
 #include "SourceCodeNavigation.h"
+#include "TargetPlatform.h"
 
 #define LOCTEXT_NAMESPACE "ProjectBrowser"
 
@@ -23,8 +24,9 @@ struct FProjectItem
 	FString ProjectFile;
 	TSharedPtr<FSlateBrush> ProjectThumbnail;
 	bool bIsNewProjectItem;
+	TArray<FName> TargetPlatforms;
 
-	FProjectItem(const FText& InName, const FText& InDescription, const FString& InEngineIdentifier, bool InUpToDate, const TSharedPtr<FSlateBrush>& InProjectThumbnail, const FString& InProjectFile, bool InIsNewProjectItem)
+	FProjectItem(const FText& InName, const FText& InDescription, const FString& InEngineIdentifier, bool InUpToDate, const TSharedPtr<FSlateBrush>& InProjectThumbnail, const FString& InProjectFile, bool InIsNewProjectItem, TArray<FName> InTargetPlatforms)
 		: Name(InName)
 		, Description(InDescription)
 		, EngineIdentifier(InEngineIdentifier)
@@ -32,6 +34,7 @@ struct FProjectItem
 		, ProjectFile(InProjectFile)
 		, ProjectThumbnail(InProjectThumbnail)
 		, bIsNewProjectItem(InIsNewProjectItem)
+		, TargetPlatforms(InTargetPlatforms)
 	{ }
 
 	/** Check if this project is up to date */
@@ -426,7 +429,7 @@ TSharedRef<SToolTip> SProjectBrowser::MakeProjectToolTip( TSharedPtr<FProjectIte
 		AddToToolTipInfoBox( InfoBox, LOCTEXT("ProjectTileTooltipPath", "Path"), FText::FromString(ProjectPath) );
 	}
 
-	if (!ProjectItem->IsUpToDate())
+	if(!ProjectItem->IsUpToDate())
 	{
 		FText Description;
 		if(FDesktopPlatformModule::Get()->IsStockEngineRelease(ProjectItem->EngineIdentifier))
@@ -450,6 +453,26 @@ TSharedRef<SToolTip> SProjectBrowser::MakeProjectToolTip( TSharedPtr<FProjectIte
 		AddToToolTipInfoBox(InfoBox, LOCTEXT("EngineVersion", "Engine"), Description);
 	}
 
+	// Create the target platform icons
+	TSharedRef<SHorizontalBox> TargetPlatformIconsBox = SNew(SHorizontalBox);
+	for(const FName& PlatformName : ProjectItem->TargetPlatforms)
+	{
+		TargetPlatformIconsBox->AddSlot()
+		.AutoWidth()
+		.HAlign(HAlign_Center)
+		.VAlign(VAlign_Center)
+		.Padding(FMargin(0, 0, 1, 0))
+		[
+			SNew(SBox)
+			.WidthOverride(20)
+			.HeightOverride(20)
+			[
+				SNew(SImage)
+				.Image(FEditorStyle::GetBrush(*FString::Printf(TEXT("Launcher.Platform_%s"), *PlatformName.ToString())))
+			]
+		];
+	}
+
 	TSharedRef<SToolTip> Tooltip = SNew(SToolTip)
 	.TextMargin(1)
 	.BorderImage( FEditorStyle::GetBrush("ProjectBrowser.TileViewTooltip.ToolTipBorder") )
@@ -474,9 +497,24 @@ TSharedRef<SToolTip> SProjectBrowser::MakeProjectToolTip( TSharedPtr<FProjectIte
 					.AutoHeight()
 					.VAlign(VAlign_Center)
 					[
-						SNew(STextBlock)
-						.Text( ProjectItem->Name )
-						.Font( FEditorStyle::GetFontStyle("ProjectBrowser.TileViewTooltip.NameFont") )
+						SNew(SHorizontalBox)
+
+						+SHorizontalBox::Slot()
+						.AutoWidth()
+						.HAlign(HAlign_Left)
+						.VAlign(VAlign_Center)
+						[
+							SNew(STextBlock)
+							.Text( ProjectItem->Name )
+							.Font( FEditorStyle::GetFontStyle("ProjectBrowser.TileViewTooltip.NameFont") )
+						]
+
+						+SHorizontalBox::Slot()
+						.HAlign(HAlign_Right)
+						.VAlign(VAlign_Center)
+						[
+							TargetPlatformIconsBox
+						]
 					]
 				]
 			]
@@ -705,6 +743,7 @@ FReply SProjectBrowser::FindProjects()
 
 	// Add all the discovered projects to the list
 	const FString EngineIdentifier = FDesktopPlatformModule::Get()->GetCurrentEngineIdentifier();
+	TArray<ITargetPlatform*> AllPlatforms = GetTargetPlatformManager()->GetTargetPlatforms();
 	for(TMap<FString, FText>::TConstIterator Iter(AbsoluteProjectFilesToCategory); Iter; ++Iter)
 	{
 		const FString& ProjectFilename = *Iter.Key();
@@ -750,8 +789,24 @@ FReply SProjectBrowser::FindProjects()
 				FString ProjectEngineIdentifier;
 				bool bIsUpToDate = FDesktopPlatformModule::Get()->GetEngineIdentifierForProject(ProjectFilename, ProjectEngineIdentifier) && ProjectEngineIdentifier == EngineIdentifier;
 
+				// Work out which platforms this project is targeting (it only shows the ones we have support for locally)
+				TArray<FName> TargetPlatforms;
+				for ( ITargetPlatform* Platform : AllPlatforms )
+				{
+					// We are only interested in standalone games
+					const bool bIsGamePlatform = !Platform->IsClientOnly() && !Platform->IsServerOnly() && !Platform->HasEditorOnlyData();
+					if ( bIsGamePlatform )
+					{
+						const FName PlatformName = *Platform->PlatformName();
+						if ( ProjectStatus.IsTargetPlatformSupported(PlatformName) )
+						{
+							TargetPlatforms.Add(PlatformName);
+						}
+					}
+				}
+
 				const bool bIsNewProjectItem = false;
-				TSharedRef<FProjectItem> NewProjectItem = MakeShareable( new FProjectItem(ProjectName, ProjectDescription, ProjectEngineIdentifier, bIsUpToDate, DynamicBrush, ProjectFilename, bIsNewProjectItem ) );
+				TSharedRef<FProjectItem> NewProjectItem = MakeShareable( new FProjectItem(ProjectName, ProjectDescription, ProjectEngineIdentifier, bIsUpToDate, DynamicBrush, ProjectFilename, bIsNewProjectItem, TargetPlatforms ) );
 				AddProjectToCategory(NewProjectItem, ProjectCategory);
 			}
 		}
diff --git a/Engine/Source/Editor/GameProjectGeneration/Public/GameProjectGenerationModule.h b/Engine/Source/Editor/GameProjectGeneration/Public/GameProjectGenerationModule.h
--- a/Engine/Source/Editor/GameProjectGeneration/Public/GameProjectGenerationModule.h
+++ b/Engine/Source/Editor/GameProjectGeneration/Public/GameProjectGenerationModule.h
@@ -65,6 +65,15 @@ public:
 	/** Warn the user if the project filename is invalid in case they renamed it outside the editor */
 	virtual void CheckAndWarnProjectFilenameValid();
 
+	/**
+	 * Update the list of supported target platforms based upon the parameters provided
+	 * This will take care of checking out and saving the updated .uproject file automatically
+	 * 
+	 * @param	InPlatformName		Name of the platform to target (eg, WindowsNoEditor)
+	 * @param	bIsSupported		true if the platform should be supported by this project, false if it should not
+	 */
+	virtual void UpdateSupportedTargetPlatforms(const FName& InPlatformName, const bool bIsSupported);
+
 private:
 	FAddCodeToProjectDialogOpenedEvent AddCodeToProjectDialogOpenedEvent;
 
diff --git a/Engine/Source/Editor/LevelEditor/LevelEditor.Build.cs b/Engine/Source/Editor/LevelEditor/LevelEditor.Build.cs
--- a/Engine/Source/Editor/LevelEditor/LevelEditor.Build.cs
+++ b/Engine/Source/Editor/LevelEditor/LevelEditor.Build.cs
@@ -44,7 +44,9 @@ public class LevelEditor : ModuleRules
 				"ContentBrowser",
                 "SceneOutliner",
                 "ActorPickerMode",
-                "RHI"
+                "RHI",
+				"Projects",
+				"TargetPlatform",
 			}
 		);
 
diff --git a/Engine/Source/Editor/LevelEditor/Private/LevelEditor.cpp b/Engine/Source/Editor/LevelEditor/Private/LevelEditor.cpp
--- a/Engine/Source/Editor/LevelEditor/Private/LevelEditor.cpp
+++ b/Engine/Source/Editor/LevelEditor/Private/LevelEditor.cpp
@@ -18,6 +18,7 @@
 #include "SlateReflector.h"
 #include "SDockTab.h"
 #include "ToolkitManager.h"
+#include "TargetPlatform.h"
 
 // @todo Editor: remove this circular dependency
 #include "Editor/MainFrame/Public/Interfaces/IMainFrameModule.h"
@@ -37,11 +38,10 @@ FLevelEditorModule::FLevelEditorModule()
 	{
 	}
 
-TSharedRef<SDockTab> SpawnLevelEditor( const FSpawnTabArgs& InArgs )
+TSharedRef<SDockTab> FLevelEditorModule::SpawnLevelEditor( const FSpawnTabArgs& InArgs )
 {
-	FLevelEditorModule& LevelEditorModule = FModuleManager::GetModuleChecked<FLevelEditorModule>( TEXT("LevelEditor") );
 	TSharedRef<SDockTab> LevelEditorTab = SNew(SDockTab) .TabRole(ETabRole::MajorTab) .ContentPadding( FMargin(0,2,0,0) );
-	LevelEditorModule.SetLevelEditorInstanceTab(LevelEditorTab);
+	SetLevelEditorInstanceTab(LevelEditorTab);
 	TSharedPtr< SWindow > OwnerWindow = InArgs.GetOwnerWindow();
 	
 	if ( !OwnerWindow.IsValid() )
@@ -54,7 +54,7 @@ TSharedRef<SDockTab> SpawnLevelEditor( const FSpawnTabArgs& InArgs )
 	{
 		TSharedPtr<SLevelEditor> LevelEditorTmp;
 		LevelEditorTab->SetContent( SAssignNew(LevelEditorTmp, SLevelEditor ) );
-		LevelEditorModule.SetLevelEditorInstance(LevelEditorTmp);
+		SetLevelEditorInstance(LevelEditorTmp);
 		LevelEditorTmp->Initialize( LevelEditorTab, OwnerWindow.ToSharedRef() );
 
 		GEditorModeTools().DeactivateAllModes();
@@ -86,6 +86,11 @@ TSharedRef<SDockTab> SpawnLevelEditor( const FSpawnTabArgs& InArgs )
 			RightContentText = FText::Format(NSLOCTEXT("UnrealEditor", "TitleBarRightContent", "{Branch}{GameName}"), Args);
 		}
 
+		// Create the target platform icons, and register our interest so we can update them if they're changed later on
+		TargetPlatformIconsBox = SNew(SHorizontalBox);
+		UpdateTargetPlatformIcons();
+		IProjectManager::Get().OnTargetPlatformsForCurrentProjectChanged().AddRaw(this, &FLevelEditorModule::UpdateTargetPlatformIcons);
+
 		RightContent =
 				SNew( SHorizontalBox )
 
@@ -102,6 +107,14 @@ TSharedRef<SDockTab> SpawnLevelEditor( const FSpawnTabArgs& InArgs )
 					]
 				]
 
+				+SHorizontalBox::Slot()
+				.AutoWidth()
+				.Padding(16.0f, 0.0f, 0.0f, 0.0f)
+				.VAlign(VAlign_Center)
+				[
+					TargetPlatformIconsBox.ToSharedRef()
+				]
+
 				+SHorizontalBox::Slot()
 				.AutoWidth()
 				.Padding(16.0f, 0.0f, 0.0f, 0.0f)
@@ -143,7 +156,7 @@ void FLevelEditorModule::StartupModule()
 
 	const IWorkspaceMenuStructure& MenuStructure = WorkspaceMenu::GetMenuStructure();
 
-	FGlobalTabmanager::Get()->RegisterTabSpawner("LevelEditor", FOnSpawnTab::CreateStatic( &SpawnLevelEditor ) )
+	FGlobalTabmanager::Get()->RegisterTabSpawner("LevelEditor", FOnSpawnTab::CreateRaw( this, &FLevelEditorModule::SpawnLevelEditor ) )
 		.SetDisplayName( NSLOCTEXT("LevelEditor", "LevelEditorTab", "Level Editor") );
 
 	FModuleManager::LoadModuleChecked<ISlateReflectorModule>("SlateReflector").RegisterTabSpawner(MenuStructure.GetDeveloperToolsCategory());
@@ -161,6 +174,8 @@ void FLevelEditorModule::StartupModule()
  */
 void FLevelEditorModule::ShutdownModule()
 {
+	IProjectManager::Get().OnTargetPlatformsForCurrentProjectChanged().RemoveAll(this);
+
 	FMessageLogModule& MessageLogModule = FModuleManager::LoadModuleChecked<FMessageLogModule>("MessageLog");
 	MessageLogModule.UnregisterLogListing("BuildAndSubmitErrors");
 
@@ -476,6 +491,49 @@ TSharedPtr<SDockTab> FLevelEditorModule::GetLevelEditorTab() const
 	return LevelEditorInstanceTabPtr.Pin().ToSharedRef();
 }
 
+void FLevelEditorModule::UpdateTargetPlatformIcons()
+{
+	TargetPlatformIconsBox->ClearChildren();
+
+	FProjectStatus ProjectStatus;
+	if(IProjectManager::Get().QueryStatusForCurrentProject(ProjectStatus))
+	{
+		TArray<ITargetPlatform*> AllPlatforms = GetTargetPlatformManager()->GetTargetPlatforms();
+		TArray<FName> TargetPlatforms;
+		for( ITargetPlatform* Platform : AllPlatforms)
+		{
+			// We are only interested in standalone games
+			const bool bIsGamePlatform = !Platform->IsClientOnly() && !Platform->IsServerOnly() && !Platform->HasEditorOnlyData();
+			if(bIsGamePlatform)
+			{
+				const FName PlatformName = *Platform->PlatformName();
+				if(ProjectStatus.IsTargetPlatformSupported(PlatformName))
+				{
+					TargetPlatforms.Add(PlatformName);
+				}
+			}
+		}
+
+		for(const FName& PlatformName : TargetPlatforms)
+		{
+			TargetPlatformIconsBox->AddSlot()
+			.AutoWidth()
+			.HAlign(HAlign_Center)
+			.VAlign(VAlign_Center)
+			.Padding(FMargin(2, 0, 2, 0))
+			[
+				SNew(SBox)
+				.WidthOverride(20)
+				.HeightOverride(20)
+				[
+					SNew(SImage)
+					.Image(FEditorStyle::GetBrush(*FString::Printf(TEXT("Launcher.Platform_%s"), *PlatformName.ToString())))
+				]
+			];
+		}
+	}
+}
+
 void FLevelEditorModule::BindGlobalLevelEditorCommands()
 {
 	check( !GlobalLevelEditorActions.IsValid() );
diff --git a/Engine/Source/Editor/LevelEditor/Public/LevelEditor.h b/Engine/Source/Editor/LevelEditor/Public/LevelEditor.h
--- a/Engine/Source/Editor/LevelEditor/Public/LevelEditor.h
+++ b/Engine/Source/Editor/LevelEditor/Public/LevelEditor.h
@@ -234,6 +234,12 @@ private:
 		FLayoutSaveRestore::SaveToConfig(GEditorLayoutIni, LayoutToSave);
 	}
 
+	/** Spawn the main level editor tab */
+	TSharedRef<SDockTab> SpawnLevelEditor( const FSpawnTabArgs& InArgs );
+
+	/** Update the target platform icons shown in the titlebar */
+	void UpdateTargetPlatformIcons();
+
 private:
 	TSharedPtr<FExtensibilityManager> MenuExtensibilityManager;
 	TSharedPtr<FExtensibilityManager> ToolBarExtensibilityManager;
@@ -280,6 +286,9 @@ private:
 	/* Holds the Editor's tab manager */
 	TSharedPtr<FTabManager> LevelEditorTabManager;
 
+	/* Container for the target platform icons for this project; updated whenever the project is updated via IProjectManager::OnTargetPlatformsForCurrentProjectChanged */
+	TSharedPtr<SHorizontalBox> TargetPlatformIconsBox;
+
 	/* Whether we have source for the level editor, and it can be recompiled. */
 	bool bCanBeRecompiled;
 };
diff --git a/Engine/Source/Editor/MainFrame/MainFrame.Build.cs b/Engine/Source/Editor/MainFrame/MainFrame.Build.cs
--- a/Engine/Source/Editor/MainFrame/MainFrame.Build.cs
+++ b/Engine/Source/Editor/MainFrame/MainFrame.Build.cs
@@ -33,7 +33,8 @@ public class MainFrame : ModuleRules
 				"WorkspaceMenuStructure",
 				"MessageLog",
 //				"SearchUI",
-				"TranslationEditor"
+				"TranslationEditor",
+				"Projects",
 			}
 		);
 
@@ -42,6 +43,7 @@ public class MainFrame : ModuleRules
 				"AssetTools",
 				"DesktopPlatform",
 				"GameProjectGeneration",
+				"ProjectTargetPlatformEditor",
 				"LevelEditor",
 				"OutputLog",
 				"Settings",
@@ -63,6 +65,7 @@ public class MainFrame : ModuleRules
 				"DesktopPlatform",
                 "Documentation",
 				"GameProjectGeneration",
+				"ProjectTargetPlatformEditor",
 				"LevelEditor",
 				"OutputLog",
                 "TranslationEditor",
diff --git a/Engine/Source/Editor/MainFrame/Private/Menus/PackageProjectMenu.h b/Engine/Source/Editor/MainFrame/Private/Menus/PackageProjectMenu.h
--- a/Engine/Source/Editor/MainFrame/Private/Menus/PackageProjectMenu.h
+++ b/Engine/Source/Editor/MainFrame/Private/Menus/PackageProjectMenu.h
@@ -6,6 +6,7 @@
 
 #pragma once
 
+#include "ProjectTargetPlatformEditor.h"
 
 #define LOCTEXT_NAMESPACE "FPackageProjectMenu"
 
@@ -30,7 +31,9 @@ public:
 		{
 			return;
 		}
-					
+		
+		IProjectTargetPlatformEditorModule& ProjectTargetPlatformEditorModule = FModuleManager::LoadModuleChecked<IProjectTargetPlatformEditorModule>("ProjectTargetPlatformEditor");
+
 		for (int32 PlatformIndex = 0; PlatformIndex < Platforms.Num(); ++PlatformIndex)
 		{
 			const ITargetPlatform* Platform = Platforms[PlatformIndex];
@@ -48,11 +51,9 @@ public:
 				CollectivePlatforms.Add(LOCTEXT("Win64","Win64"));
 
 				MenuBuilder.AddSubMenu(
-					Platform->DisplayName(),
-					Platform->DisplayName(),
+					ProjectTargetPlatformEditorModule.MakePlatformMenuItemWidget(Platform), 
 					FNewMenuDelegate::CreateStatic(&FPackageProjectMenu::AddPlatformSubPlatformsToMenu, Platform, CollectivePlatforms),
-					false,
-					FSlateIcon(FEditorStyle::GetStyleSetName(), *FString::Printf(TEXT("Launcher.Platform_%s"), *Platform->PlatformName()))
+					false
 					);
 			}
 			else
@@ -70,6 +71,8 @@ public:
 
 		MenuBuilder.AddMenuSeparator();
 		MenuBuilder.AddMenuEntry(FMainFrameCommands::Get().PackagingSettings);
+
+		ProjectTargetPlatformEditorModule.AddOpenProjectTargetPlatformEditorMenuItem(MenuBuilder);
 	}
 
 protected:
@@ -82,16 +85,31 @@ protected:
 	*/
 	static void AddPlatformToMenu(FMenuBuilder& MenuBuilder, const ITargetPlatform* Platform)
 	{
+		IProjectTargetPlatformEditorModule& ProjectTargetPlatformEditorModule = FModuleManager::LoadModuleChecked<IProjectTargetPlatformEditorModule>("ProjectTargetPlatformEditor");
+
 		FUIAction Action(
 			FExecuteAction::CreateStatic(&FMainFrameActionCallbacks::PackageProject, Platform->PlatformName(), Platform->DisplayName(), FString()),
 			FCanExecuteAction::CreateStatic(&FMainFrameActionCallbacks::PackageProjectCanExecute, Platform->PlatformName(), Platform->SupportsFeature(ETargetPlatformFeatures::Packaging))
 			);
 
+		// ... generate tooltip text
+		FFormatNamedArguments TooltipArguments;
+		TooltipArguments.Add(TEXT("DisplayName"), Platform->DisplayName());
+		FText Tooltip = FText::Format(LOCTEXT("PackageGameForPlatformTooltip", "Build, cook and package your game for the {DisplayName} platform"), TooltipArguments);
+
+		FProjectStatus ProjectStatus;
+		if (IProjectManager::Get().QueryStatusForCurrentProject(ProjectStatus) && !ProjectStatus.IsTargetPlatformSupported(*Platform->PlatformName()))
+		{
+			FText TooltipLine2 = FText::Format(LOCTEXT("LaunchDevicePlatformWarning", "{DisplayName} is not listed as a target platform for this project, so may not run as expected."), TooltipArguments);
+			Tooltip = FText::Format(FText::FromString(TEXT("{0}\n\n{1}")), Tooltip, TooltipLine2);
+		}
+
+		// ... and add a menu entry
 		MenuBuilder.AddMenuEntry(
-			Platform->DisplayName(),
-			FText::Format(LOCTEXT("PackageGameForPlatformTooltip", "Build, cook and package your game for the {0} platform"), Platform->DisplayName()),
-			FSlateIcon(FEditorStyle::GetStyleSetName(), *FString::Printf(TEXT("Launcher.Platform_%s"), *Platform->PlatformName())),
-			Action
+			Action, 
+			ProjectTargetPlatformEditorModule.MakePlatformMenuItemWidget(Platform), 
+			NAME_None, 
+			Tooltip
 			);
 	}
 
@@ -106,6 +124,8 @@ protected:
 	*/
 	static void AddPlatformSubPlatformsToMenu(FMenuBuilder& MenuBuilder, const ITargetPlatform* Platform, TArray<FText> DisplayNames)
 	{
+		IProjectTargetPlatformEditorModule& ProjectTargetPlatformEditorModule = FModuleManager::LoadModuleChecked<IProjectTargetPlatformEditorModule>("ProjectTargetPlatformEditor");
+
 		for (const FText& DisplayName : DisplayNames)
 		{
 			FString OptionalParams = FString::Printf(TEXT("-targetplatform=%s"), *DisplayName.ToString());
@@ -114,11 +134,25 @@ protected:
 				FCanExecuteAction::CreateStatic(&FMainFrameActionCallbacks::PackageProjectCanExecute, Platform->PlatformName(), Platform->SupportsFeature(ETargetPlatformFeatures::Packaging))
 				);
 
+			// ... generate tooltip text
+			FFormatNamedArguments TooltipArguments;
+			TooltipArguments.Add(TEXT("DisplayName"), Platform->DisplayName());
+			TooltipArguments.Add(TEXT("DisplaySubName"), DisplayName);
+			FText Tooltip = FText::Format(LOCTEXT("PackageGameForSubPlatformTooltip", "Build, cook and package your game for the {DisplayName} | {DisplaySubName} platform"), TooltipArguments);
+
+			FProjectStatus ProjectStatus;
+			if (IProjectManager::Get().QueryStatusForCurrentProject(ProjectStatus) && !ProjectStatus.IsTargetPlatformSupported(*Platform->PlatformName()))
+			{
+				FText TooltipLine2 = FText::Format(LOCTEXT("LaunchDevicePlatformWarning", "{DisplayName} is not listed as a target platform for this project, so may not run as expected."), TooltipArguments);
+				Tooltip = FText::Format(FText::FromString(TEXT("{0}\n\n{1}")), Tooltip, TooltipLine2);
+			}
+
+			// ... and add a menu entry
 			MenuBuilder.AddMenuEntry(
-				DisplayName,
-				FText::Format(LOCTEXT("PackageGameForPlatformTooltip", "Build, cook and package your game for the {0} | {1} platform"), Platform->DisplayName(), DisplayName),
-				FSlateIcon(FEditorStyle::GetStyleSetName(), *FString::Printf(TEXT("Launcher.Platform_%s"), *Platform->PlatformName())),
-				Action
+				Action, 
+				ProjectTargetPlatformEditorModule.MakePlatformMenuItemWidget(Platform, false, DisplayName), 
+				NAME_None, 
+				Tooltip
 				);
 		}
 	}
diff --git a/Engine/Source/Editor/ProjectSettingsViewer/Private/ProjectSettingsViewerModule.cpp b/Engine/Source/Editor/ProjectSettingsViewer/Private/ProjectSettingsViewerModule.cpp
--- a/Engine/Source/Editor/ProjectSettingsViewer/Private/ProjectSettingsViewerModule.cpp
+++ b/Engine/Source/Editor/ProjectSettingsViewer/Private/ProjectSettingsViewerModule.cpp
@@ -6,6 +6,7 @@
 
 #include "ProjectSettingsViewerPrivatePCH.h"
 #include "Engine/Console.h"
+#include "ProjectTargetPlatformEditor.h"
 
 #define LOCTEXT_NAMESPACE "FProjectSettingsViewerModule"
 
@@ -179,6 +180,14 @@ protected:
 			GetMutableDefault<UProjectPackagingSettings>()
 		);
 
+		// platforms settings
+		TWeakPtr<SWidget> ProjectTargetPlatformEditorPanel = FModuleManager::LoadModuleChecked<IProjectTargetPlatformEditorModule>("ProjectTargetPlatformEditor").CreateProjectTargetPlatformEditorPanel();
+		SettingsModule.RegisterSettings("Project", "Game", "TargetPlatforms",
+			LOCTEXT("ProjectTargetPlatformsSettingsName", "Target Platforms"),
+			LOCTEXT("ProjectTargetPlatformsSettingsDescription", "Specify which platforms your project is targeting."),
+			ProjectTargetPlatformEditorPanel.Pin().ToSharedRef()
+		);
+
 		// movie settings
 		SettingsModule.RegisterSettings("Project", "Game", "Movies",
 			LOCTEXT("MovieSettingsName", "Movies"),
@@ -226,6 +235,7 @@ protected:
 			SettingsModule->UnregisterSettings("Project", "Game", "General");
 			SettingsModule->UnregisterSettings("Project", "Game", "Maps");
 			SettingsModule->UnregisterSettings("Project", "Game", "Packaging");
+			SettingsModule->UnregisterSettings("Project", "Game", "TargetPlatforms");
 			SettingsModule->UnregisterSettings("Project", "Game", "Movies");
 //			SettingsModule->UnregisterSettings("Project", "Game", "GameSession");
 //			SettingsModule->UnregisterSettings("Project", "Game", "HUD");
diff --git a/Engine/Source/Editor/ProjectSettingsViewer/ProjectSettingsViewer.Build.cs b/Engine/Source/Editor/ProjectSettingsViewer/ProjectSettingsViewer.Build.cs
--- a/Engine/Source/Editor/ProjectSettingsViewer/ProjectSettingsViewer.Build.cs
+++ b/Engine/Source/Editor/ProjectSettingsViewer/ProjectSettingsViewer.Build.cs
@@ -23,7 +23,8 @@ namespace UnrealBuildTool.Rules
 					"Slate",
 					"SlateCore",
 					"UnrealEd",
-					"MoviePlayer"
+					"MoviePlayer",
+					"ProjectTargetPlatformEditor",
 				}
 			);
 
diff --git a/Engine/Source/Editor/ProjectTargetPlatformEditor/Private/ProjectTargetPlatformEditorModule.cpp b/Engine/Source/Editor/ProjectTargetPlatformEditor/Private/ProjectTargetPlatformEditorModule.cpp
new file mode 100644
--- /dev/null
+++ b/Engine/Source/Editor/ProjectTargetPlatformEditor/Private/ProjectTargetPlatformEditorModule.cpp
@@ -0,0 +1,144 @@
+// Copyright 1998-2014 Epic Games, Inc. All Rights Reserved.
+
+#include "ProjectTargetPlatformEditorPrivatePCH.h"
+#include "Settings.h"
+
+#define LOCTEXT_NAMESPACE "FProjectTargetPlatformEditorModule"
+
+/**
+ * Implements the platform target platform editor module
+ */
+class FProjectTargetPlatformEditorModule : public IProjectTargetPlatformEditorModule
+{
+public:
+
+	// Begin IProjectTargetPlatformEditorModule interface
+
+	virtual TWeakPtr<SWidget> CreateProjectTargetPlatformEditorPanel() OVERRIDE
+	{
+		TSharedPtr<SWidget> Panel = SNew(SProjectTargetPlatformSettings);
+		EditorPanels.Add(Panel);
+
+		return Panel;
+	}
+
+	virtual void DestroyProjectTargetPlatformEditorPanel(const TWeakPtr<SWidget>& Panel) OVERRIDE
+	{
+		EditorPanels.Remove(Panel.Pin());
+	}
+
+	virtual void AddOpenProjectTargetPlatformEditorMenuItem(FMenuBuilder& MenuBuilder) const OVERRIDE
+	{
+		struct Local
+		{
+			static void OpenSettings( FName ContainerName, FName CategoryName, FName SectionName )
+			{
+				FModuleManager::LoadModuleChecked<ISettingsModule>("Settings").ShowViewer(ContainerName, CategoryName, SectionName);
+			}
+		};
+
+		MenuBuilder.AddMenuEntry(
+			LOCTEXT("TargetPlatformsMenuLabel", "Target Platforms..."),
+			LOCTEXT("ProjectSettingsMenuToolTip", "Change which platforms this project is targeting"),
+			FSlateIcon(),
+			FUIAction(FExecuteAction::CreateStatic(&Local::OpenSettings, FName("Project"), FName("Game"), FName("TargetPlatforms")))
+			);
+	}
+
+	virtual TSharedRef<SWidget> MakePlatformMenuItemWidget(const ITargetPlatform* const Platform, const bool bForCheckBox = false, const FText& DisplayNameOverride = FText()) const OVERRIDE
+	{
+		struct Local
+		{
+			static EVisibility IsUnsupportedPlatformWarningVisible(const FName PlatformName)
+			{
+				FProjectStatus ProjectStatus;
+				return (!IProjectManager::Get().QueryStatusForCurrentProject(ProjectStatus) || ProjectStatus.IsTargetPlatformSupported(PlatformName)) ? EVisibility::Hidden : EVisibility::Visible;
+			}
+		};
+
+		return 
+			SNew(SHorizontalBox)
+			+SHorizontalBox::Slot()
+			.AutoWidth()
+			.Padding(FMargin((bForCheckBox) ? 2 : 13, 0, 2, 0))
+			.HAlign(HAlign_Center)
+			.VAlign(VAlign_Center)
+			[
+				SNew(SOverlay)
+				+SOverlay::Slot()
+				.HAlign(HAlign_Left)
+				.VAlign(VAlign_Center)
+				[
+					SNew(SBox)
+					.WidthOverride(MultiBoxConstants::MenuIconSize)
+					.HeightOverride(MultiBoxConstants::MenuIconSize)
+					[
+						SNew(SImage)
+						.Image(FEditorStyle::GetBrush(*FString::Printf(TEXT("Launcher.Platform_%s"), *Platform->PlatformName())))
+					]
+				]
+				+SOverlay::Slot()
+				.Padding(FMargin(MultiBoxConstants::MenuIconSize * 0.5f, 0, 0, 0))
+				.HAlign(HAlign_Left)
+				.VAlign(VAlign_Bottom)
+				[
+					SNew(SBox)
+					.WidthOverride(MultiBoxConstants::MenuIconSize * 0.9f)
+					.HeightOverride(MultiBoxConstants::MenuIconSize * 0.9f)
+					[
+						SNew(SImage)
+						.Visibility(TAttribute<EVisibility>::Create(TAttribute<EVisibility>::FGetter::CreateStatic(&Local::IsUnsupportedPlatformWarningVisible, FName(*Platform->PlatformName()))))
+						.Image(FEditorStyle::GetBrush("Launcher.Platform.Warning"))
+					]
+				]
+			]
+			+SHorizontalBox::Slot()
+			.FillWidth(1.0f)
+			.Padding(FMargin((bForCheckBox) ? 2 : 7, 0, 6, 0))
+			.VAlign(VAlign_Center)
+			[
+				SNew(STextBlock)
+				.TextStyle(FEditorStyle::Get(), "Menu.Label")
+				.Text((DisplayNameOverride.IsEmpty()) ? Platform->DisplayName() : DisplayNameOverride)
+			];
+	}
+
+	virtual bool ShowUnsupportedTargetWarning(const FName PlatformName) const OVERRIDE
+	{
+		// Don't show the warning during automation testing; the dlg is modal and blocks
+		FProjectStatus ProjectStatus;
+		if(!GIsAutomationTesting && IProjectManager::Get().QueryStatusForCurrentProject(ProjectStatus) && !ProjectStatus.IsTargetPlatformSupported(PlatformName))
+		{
+			ITargetPlatform *const Platform = GetTargetPlatformManager()->FindTargetPlatform(PlatformName.ToString());
+			check(Platform);
+
+			FFormatNamedArguments Args;
+			Args.Add(TEXT("DisplayName"), Platform->DisplayName());
+			FText WarningText = FText::Format(LOCTEXT("ShowUnsupportedTargetWarning_Message", "{DisplayName} is not listed as a target platform for this project, so may not run as expected.\n\nDo you wish to continue?"), Args);
+
+			FSuppressableWarningDialog::FSetupInfo Info(
+				WarningText, 
+				LOCTEXT("ShowUnsupportedTargetWarning_Title", "Unsupported Platform"), 
+				TEXT("SuppressUnsupportedTargetWarningDialog")
+				);
+			Info.ConfirmText = LOCTEXT("ShowUnsupportedTargetWarning_Confirm", "Continue");
+			Info.CancelText = LOCTEXT("ShowUnsupportedTargetWarning_Cancel", "Cancel");
+			FSuppressableWarningDialog UnsupportedTargetWarningDialog(Info);
+
+			return UnsupportedTargetWarningDialog.ShowModal() != FSuppressableWarningDialog::EResult::Cancel;
+		}
+
+		return true;
+	}
+
+	// End IProjectTargetPlatformEditorModule interface
+
+private:
+
+	// Holds the collection of created editor panels.
+	TArray<TSharedPtr<SWidget> > EditorPanels;
+};
+
+#undef LOCTEXT_NAMESPACE
+
+IMPLEMENT_MODULE(FProjectTargetPlatformEditorModule, ProjectTargetPlatformEditor);
diff --git a/Engine/Source/Editor/ProjectTargetPlatformEditor/Private/ProjectTargetPlatformEditorPrivatePCH.h b/Engine/Source/Editor/ProjectTargetPlatformEditor/Private/ProjectTargetPlatformEditorPrivatePCH.h
new file mode 100644
--- /dev/null
+++ b/Engine/Source/Editor/ProjectTargetPlatformEditor/Private/ProjectTargetPlatformEditorPrivatePCH.h
@@ -0,0 +1,22 @@
+// Copyright 1998-2014 Epic Games, Inc. All Rights Reserved.
+
+#pragma once
+
+
+#include "ProjectTargetPlatformEditor.h"
+
+
+/* Dependencies
+ *****************************************************************************/
+
+#include "UnrealEd.h"
+#include "ModuleManager.h"
+#include "EditorStyle.h"
+#include "TargetPlatform.h"
+#include "GameProjectGenerationModule.h"
+
+
+/* Private includes
+ *****************************************************************************/
+
+#include "SProjectTargetPlatformSettings.h"
diff --git a/Engine/Source/Editor/ProjectTargetPlatformEditor/Private/ProjectTargetPlatformEditorWidgets.cpp b/Engine/Source/Editor/ProjectTargetPlatformEditor/Private/ProjectTargetPlatformEditorWidgets.cpp
new file mode 100644
--- /dev/null
+++ b/Engine/Source/Editor/ProjectTargetPlatformEditor/Private/ProjectTargetPlatformEditorWidgets.cpp
@@ -0,0 +1,3 @@
+// Copyright 1998-2014 Epic Games, Inc. All Rights Reserved.
+
+#include "ProjectTargetPlatformEditorPrivatePCH.h"
diff --git a/Engine/Source/Editor/ProjectTargetPlatformEditor/Private/Widgets/SProjectTargetPlatformSettings.cpp b/Engine/Source/Editor/ProjectTargetPlatformEditor/Private/Widgets/SProjectTargetPlatformSettings.cpp
new file mode 100644
--- /dev/null
+++ b/Engine/Source/Editor/ProjectTargetPlatformEditor/Private/Widgets/SProjectTargetPlatformSettings.cpp
@@ -0,0 +1,200 @@
+// Copyright 1998-2014 Epic Games, Inc. All Rights Reserved.
+
+#include "ProjectTargetPlatformEditorPrivatePCH.h"
+
+#define LOCTEXT_NAMESPACE "SProjectTargetPlatformSettings"
+
+/**
+ * Implements a row widget for the target platform list view
+ */
+class SProjectTargetPlatformsListRow : public SMultiColumnTableRow<TSharedPtr<FProjectTargetPlatform>>
+{
+public:
+
+	SLATE_BEGIN_ARGS(SProjectTargetPlatformsListRow)
+		: _TargetPlatform()
+	{}
+
+	SLATE_ARGUMENT(TSharedPtr<FProjectTargetPlatform>, TargetPlatform)
+
+	SLATE_END_ARGS()
+
+public:
+
+	/**
+	 * Constructs the widget.
+	 *
+	 * @param InArgs - The arguments.
+	 */
+	void Construct(const FArguments& InArgs, const TSharedRef<STableViewBase>& InOwnerTableView)
+	{
+		check(InArgs._TargetPlatform.IsValid());
+		TargetPlatform = InArgs._TargetPlatform;
+
+		SMultiColumnTableRow<TSharedPtr<FProjectTargetPlatform>>::Construct(FSuperRowType::FArguments(), InOwnerTableView);
+	}
+
+public:
+
+	/**
+	 * Generates the widget for the specified column.
+	 *
+	 * @param ColumnName - The name of the column to generate the widget for.
+	 *
+	 * @return The widget.
+	 */
+	BEGIN_SLATE_FUNCTION_BUILD_OPTIMIZATION
+	virtual TSharedRef<SWidget> GenerateWidgetForColumn(const FName& ColumnName) OVERRIDE
+	{
+		if(ColumnName == "Enabled")
+		{
+			return SNew(SBox)
+			.Padding(2.0f)
+			.VAlign(VAlign_Center)
+			[
+				SNew(SCheckBox)
+				.IsChecked(this, &SProjectTargetPlatformsListRow::HandleEnabledCheckBoxIsChecked)
+				.OnCheckStateChanged(this, &SProjectTargetPlatformsListRow::HandleEnabledCheckBoxStateChanged)
+			];
+		}
+		else if(ColumnName == "Icon")
+		{
+			return SNew(SBox)
+			.Padding(FMargin(4.0f, 0.0f))
+			[
+				SNew(SImage)
+				.Image(TargetPlatform->Icon)
+			];
+		}
+		else if(ColumnName == "Platform")
+		{
+			return SNew(SBox)
+			.Padding(FMargin(4.0f, 0.0f))
+			.VAlign(VAlign_Center)
+			[
+				SNew(STextBlock)
+				.Text(TargetPlatform->DisplayName)
+				.ToolTipText(TargetPlatform->PlatformName.ToString())
+			];
+		}
+
+		return SNullWidget::NullWidget;
+	}
+	END_SLATE_FUNCTION_BUILD_OPTIMIZATION
+
+private:
+	/** Check to see if the "enabled" checkbox should be checked for this platform */
+	ESlateCheckBoxState::Type HandleEnabledCheckBoxIsChecked() const
+	{
+		// We need to check that the platform is explicitly supported as a target, rather than implicitly supported due to an empty list
+		FProjectStatus ProjectStatus;
+		return (IProjectManager::Get().QueryStatusForCurrentProject(ProjectStatus) && ProjectStatus.IsTargetPlatformSupported(TargetPlatform->PlatformName, false/*bAllowSupportedIfEmptyList*/)) 
+			? ESlateCheckBoxState::Checked 
+			: ESlateCheckBoxState::Unchecked;
+	}
+
+	/** Handle the "enabled" checkbox state being changed for this platform */
+	void HandleEnabledCheckBoxStateChanged(ESlateCheckBoxState::Type InState)
+	{
+		FGameProjectGenerationModule::Get().UpdateSupportedTargetPlatforms(TargetPlatform->PlatformName, InState == ESlateCheckBoxState::Checked);
+	}
+
+	TSharedPtr<FProjectTargetPlatform> TargetPlatform;
+};
+
+BEGIN_SLATE_FUNCTION_BUILD_OPTIMIZATION
+void SProjectTargetPlatformSettings::Construct(const FArguments& InArgs)
+{
+	{
+		TArray<ITargetPlatform*> AllPlatforms = GetTargetPlatformManager()->GetTargetPlatforms();
+		for(ITargetPlatform* Platform : AllPlatforms)
+		{
+			// We are only interested in standalone games
+			const bool bIsGamePlatform = !Platform->IsClientOnly() && !Platform->IsServerOnly() && !Platform->HasEditorOnlyData();
+			if(bIsGamePlatform)
+			{
+				// We use FProjectTargetPlatform rather than ITargetPlatform* to avoid trying to access ITargetPlatform
+				// instances after their module has been unloaded (which can happen on editor shutdown)
+				AvailablePlatforms.Add(MakeShareable(new FProjectTargetPlatform(
+					Platform->DisplayName(), 
+					*Platform->PlatformName(), 
+					FEditorStyle::GetBrush(*FString::Printf(TEXT("Launcher.Platform_%s"), *Platform->PlatformName()))
+					)));
+			}
+		}
+	}
+
+	ChildSlot
+	[
+		SNew(SVerticalBox)
+
+		+SVerticalBox::Slot()
+		.AutoHeight()
+		[
+			SNew(SBorder)
+			.BorderImage(FEditorStyle::GetBrush("ToolPanel.GroupBorder"))
+			.Padding(5.0f)
+			[
+				SNew(SListView<TSharedPtr<FProjectTargetPlatform>>)
+				.ItemHeight(20.0f)
+				.ListItemsSource(&AvailablePlatforms)
+				.OnGenerateRow(this, &SProjectTargetPlatformSettings::HandlePlatformsListViewGenerateRow)
+				.SelectionMode(ESelectionMode::Single)
+				.HeaderRow
+				(
+					SNew(SHeaderRow)
+
+					+SHeaderRow::Column("Enabled")
+					.DefaultLabel(FText::FromString(TEXT(" ")))
+					.FixedWidth(28.0f)
+					.HAlignCell(HAlign_Center)
+					.HAlignHeader(HAlign_Center)
+
+					+SHeaderRow::Column("Icon")
+					.DefaultLabel(FText::FromString(TEXT(" ")))
+					.FixedWidth(32.0f)
+
+					+SHeaderRow::Column("Platform")
+					.DefaultLabel(LOCTEXT("PlatformsListPlatformColumnHeader", "Platform"))
+				)
+			]
+		]
+
+		+SVerticalBox::Slot()
+		.AutoHeight()
+		.Padding(FMargin(0.0f, 5.0f))
+		[
+			SNew(SBorder)
+			.BorderImage(FEditorStyle::GetBrush("ToolPanel.GroupBorder"))
+			.Padding(5.0f)
+			[
+				SNew(SVerticalBox)
+
+				+SVerticalBox::Slot()
+				.AutoHeight()
+				[
+					SNew(STextBlock)
+					.AutoWrapText(true)
+					.Text(LOCTEXT("PlatformsListDescription", "Use the list above to select the platforms that your project is targeting. Attempting to run your project on a platform that hasn't been selected will result in a warning stating that the project may not run as expected."))
+				]
+
+				+SVerticalBox::Slot()
+				.AutoHeight()
+				[
+					SNew(STextBlock)
+					.AutoWrapText(true)
+					.Text(LOCTEXT("PlatformsListNote", "Please note that selecting no platforms is considered the same as selecting all platforms."))
+				]
+			]
+		]
+	];
+}
+END_SLATE_FUNCTION_BUILD_OPTIMIZATION
+
+TSharedRef<ITableRow> SProjectTargetPlatformSettings::HandlePlatformsListViewGenerateRow(TSharedPtr<FProjectTargetPlatform> TargetPlatform, const TSharedRef<STableViewBase>& OwnerTable)
+{
+	return SNew(SProjectTargetPlatformsListRow, OwnerTable)
+		.TargetPlatform(TargetPlatform);
+}
+
+#undef LOCTEXT_NAMESPACE
diff --git a/Engine/Source/Editor/ProjectTargetPlatformEditor/Private/Widgets/SProjectTargetPlatformSettings.h b/Engine/Source/Editor/ProjectTargetPlatformEditor/Private/Widgets/SProjectTargetPlatformSettings.h
new file mode 100644
--- /dev/null
+++ b/Engine/Source/Editor/ProjectTargetPlatformEditor/Private/Widgets/SProjectTargetPlatformSettings.h
@@ -0,0 +1,40 @@
+// Copyright 1998-2014 Epic Games, Inc. All Rights Reserved.
+
+#pragma once
+
+struct FProjectTargetPlatform
+{
+	FProjectTargetPlatform(FText InDisplayName, FName InPlatformName, const FSlateBrush* InIcon)
+		: DisplayName(InDisplayName)
+		, PlatformName(InPlatformName)
+		, Icon(InIcon)
+	{
+	}
+
+	/** Localized display name to show in the UI, eg, Windows */
+	FText DisplayName;
+
+	/** Internal platform name to use as a key, eg, WindowsNoEditor */
+	FName PlatformName;
+
+	/** Icon representing this platform */
+	const FSlateBrush* Icon;
+};
+
+class SProjectTargetPlatformSettings : public SCompoundWidget
+{
+public:
+	SLATE_BEGIN_ARGS(SProjectTargetPlatformSettings)
+	{}
+	SLATE_END_ARGS()
+
+	/** Constructs this widget with InArgs */
+	void Construct(const FArguments& InArgs);
+
+private:
+	/** Generate a row for the available platforms list */
+	TSharedRef<ITableRow> HandlePlatformsListViewGenerateRow(TSharedPtr<FProjectTargetPlatform> TargetPlatform, const TSharedRef<STableViewBase>& OwnerTable);
+
+	/** List of available target platforms */
+	TArray<TSharedPtr<FProjectTargetPlatform>> AvailablePlatforms;
+};
diff --git a/Engine/Source/Editor/ProjectTargetPlatformEditor/ProjectTargetPlatformEditor.Build.cs b/Engine/Source/Editor/ProjectTargetPlatformEditor/ProjectTargetPlatformEditor.Build.cs
new file mode 100644
--- /dev/null
+++ b/Engine/Source/Editor/ProjectTargetPlatformEditor/ProjectTargetPlatformEditor.Build.cs
@@ -0,0 +1,40 @@
+// Copyright 1998-2014 Epic Games, Inc. All Rights Reserved.
+
+using UnrealBuildTool;
+
+public class ProjectTargetPlatformEditor : ModuleRules
+{
+	public ProjectTargetPlatformEditor(TargetInfo Target)
+	{
+		PublicDependencyModuleNames.AddRange(
+			new string[]
+			{
+				"Core",
+				"Slate",
+                "SlateCore",
+			}
+		);
+
+		PrivateDependencyModuleNames.AddRange(
+			new string[]
+			{
+				"CoreUObject",
+				"InputCore",
+				"EditorStyle",
+				"TargetPlatform",
+				"GameProjectGeneration",
+				"Settings",
+				"UnrealEd",
+				"Projects",
+			}
+		);
+
+		PrivateIncludePaths.AddRange(
+			new string[]
+			{
+				"Editor/ProjectTargetPlatformEditor/Private",
+				"Editor/ProjectTargetPlatformEditor/Private/Widgets",
+			}
+		);
+	}
+}
diff --git a/Engine/Source/Editor/ProjectTargetPlatformEditor/Public/Interfaces/IProjectTargetPlatformEditorModule.h b/Engine/Source/Editor/ProjectTargetPlatformEditor/Public/Interfaces/IProjectTargetPlatformEditorModule.h
new file mode 100644
--- /dev/null
+++ b/Engine/Source/Editor/ProjectTargetPlatformEditor/Public/Interfaces/IProjectTargetPlatformEditorModule.h
@@ -0,0 +1,57 @@
+// Copyright 1998-2014 Epic Games, Inc. All Rights Reserved.
+
+#pragma once
+
+/**
+ * Interface for the platform target platform editor module
+ */
+class IProjectTargetPlatformEditorModule : public IModuleInterface
+{
+public:
+
+	/**
+	 * Creates an project target platform editor panel widget.
+	 *
+	 * @return The created widget.
+	 */
+	virtual TWeakPtr<SWidget> CreateProjectTargetPlatformEditorPanel() = 0;
+
+	/**
+	 * Destroys a previously created editor panel widget.
+	 *
+	 * @param Panel - The panel to destroy.
+	 */
+	virtual void DestroyProjectTargetPlatformEditorPanel(const TWeakPtr<SWidget>& Panel) = 0;
+
+	/**
+	 * Adds a menu item to open the target platform editor panel in the project settings.
+	 *
+	 * @param MenuBuilder - The menu builder to add the item to.
+	 */
+	virtual void AddOpenProjectTargetPlatformEditorMenuItem(FMenuBuilder& MenuBuilder) const = 0;
+
+	/**
+	 * Creates the widget to use for a platform entry within a FMenuBuilder.
+	 *
+	 * @param Platform - The target platform to build the widget for
+	 * @param bForCheckBox - true if the widget is for a checkbox menu item, false if it's for any other menu item (this affects the padding)
+	 * @param DisplayNameOverride - An alternate name to use for the platform, otherwise Platform->DisplayName() will be used
+	 */
+	virtual TSharedRef<SWidget> MakePlatformMenuItemWidget(const ITargetPlatform* const Platform, const bool bForCheckBox = false, const FText& DisplayNameOverride = FText()) const = 0;
+
+	/**
+	 * Check to see if the given platform is on the list of supported targets, and show a warning message if it's not, allowing the user to continue or cancel
+	 *
+	 * @param PlatformName - The name of the platform to test.
+	 *
+	 * @return true if the platform is supported, or the warning was accepted with "continue".
+	 */
+	virtual bool ShowUnsupportedTargetWarning(const FName PlatformName) const = 0;
+
+public:
+
+	/**
+	 * Virtual destructor.
+	 */
+	virtual ~IProjectTargetPlatformEditorModule() { }
+};
diff --git a/Engine/Source/Editor/ProjectTargetPlatformEditor/Public/ProjectTargetPlatformEditor.h b/Engine/Source/Editor/ProjectTargetPlatformEditor/Public/ProjectTargetPlatformEditor.h
new file mode 100644
--- /dev/null
+++ b/Engine/Source/Editor/ProjectTargetPlatformEditor/Public/ProjectTargetPlatformEditor.h
@@ -0,0 +1,16 @@
+// Copyright 1998-2014 Epic Games, Inc. All Rights Reserved.
+
+#pragma once
+
+
+/* Dependencies
+ *****************************************************************************/
+
+#include "ModuleManager.h"
+#include "Slate.h"
+
+
+/* Interfaces
+ *****************************************************************************/
+
+#include "IProjectTargetPlatformEditorModule.h"
diff --git a/Engine/Source/Editor/UnrealEd/Private/Kismet2/DebuggerCommands.cpp b/Engine/Source/Editor/UnrealEd/Private/Kismet2/DebuggerCommands.cpp
--- a/Engine/Source/Editor/UnrealEd/Private/Kismet2/DebuggerCommands.cpp
+++ b/Engine/Source/Editor/UnrealEd/Private/Kismet2/DebuggerCommands.cpp
@@ -17,6 +17,7 @@
 #include "Runtime/Analytics/Analytics/Public/Interfaces/IAnalyticsProvider.h"
 
 #include "GameProjectGenerationModule.h"
+#include "ProjectTargetPlatformEditor.h"
 
 //@TODO: Remove this dependency
 #include "Editor/LevelEditor/Public/LevelEditor.h"
@@ -556,6 +557,7 @@ TSharedRef< SWidget > FPlayWorldCommands::GenerateLaunchMenuContent( TSharedRef<
 	// shared devices section
 	TSharedPtr<ITargetDeviceServicesModule> TargetDeviceServicesModule = StaticCastSharedPtr<ITargetDeviceServicesModule>(FModuleManager::Get().LoadModule(TEXT("TargetDeviceServices")));
 	TArray<ITargetPlatform*> Platforms = GetTargetPlatformManager()->GetTargetPlatforms();
+	IProjectTargetPlatformEditorModule& ProjectTargetPlatformEditorModule = FModuleManager::LoadModuleChecked<IProjectTargetPlatformEditorModule>("ProjectTargetPlatformEditor");
 
 	TArray<FString> PlatformsToMaybeInstallLinksFor;
 	PlatformsToMaybeInstallLinksFor.Add(TEXT("Android"));
@@ -625,18 +627,24 @@ TSharedRef< SWidget > FPlayWorldCommands::GenerateLaunchMenuContent( TSharedRef<
 					// ... generate tooltip text
 					FFormatNamedArguments TooltipArguments;
 					TooltipArguments.Add(TEXT("DeviceID"), FText::FromString(DeviceProxy->GetName()));
-					TooltipArguments.Add(TEXT("PlatformName"), FText::FromString(Platform->PlatformName()));
-					FText Tooltip = FText::Format(LOCTEXT("LaunchDeviceToolTipText", "Launch the game on this {PlatformName} device ({DeviceID})"), TooltipArguments);
+					TooltipArguments.Add(TEXT("DisplayName"), Platform->DisplayName());
+					FText Tooltip = FText::Format(LOCTEXT("LaunchDeviceToolTipText", "Launch the game on this {DisplayName} device ({DeviceID})"), TooltipArguments);
+
+					FProjectStatus ProjectStatus;
+					if (IProjectManager::Get().QueryStatusForCurrentProject(ProjectStatus) && !ProjectStatus.IsTargetPlatformSupported(*Platform->PlatformName())) 
+					{
+						FText TooltipLine2 = FText::Format(LOCTEXT("LaunchDevicePlatformWarning", "{DisplayName} is not listed as a target platform for this project, so may not run as expected."), TooltipArguments);
+						Tooltip = FText::Format(FText::FromString(TEXT("{0}\n\n{1}")), Tooltip, TooltipLine2);
+					}
 
 					// ... and add a menu entry
 					MenuBuilder.AddMenuEntry(
-						Label,
-						Tooltip,
-						FSlateIcon(FEditorStyle::GetStyleSetName(), *IconName),
-						LaunchDeviceAction,
-						NAME_None,
+						LaunchDeviceAction, 
+						ProjectTargetPlatformEditorModule.MakePlatformMenuItemWidget(Platform, true, Label), 
+						NAME_None, 
+						Tooltip, 
 						EUserInterfaceActionType::Check
-					);
+						);
 				}
 			}
 		}
@@ -687,6 +695,8 @@ TSharedRef< SWidget > FPlayWorldCommands::GenerateLaunchMenuContent( TSharedRef<
 	MenuBuilder.BeginSection("LevelEditorLaunchOptions");
 	{
 		MenuBuilder.AddMenuEntry( FPlayWorldCommands::Get().OpenDeviceManager );
+		
+		ProjectTargetPlatformEditorModule.AddOpenProjectTargetPlatformEditorMenuItem(MenuBuilder);
 	}
 	MenuBuilder.EndSection();
 
@@ -1539,7 +1549,14 @@ bool FInternalPlayWorldCommandCallbacks::CanLaunchOnDevice( const FString& Devic
 
 void FInternalPlayWorldCommandCallbacks::LaunchOnDevice( const FString& DeviceId, const FString& DeviceName )
 {
-	GUnrealEd->RequestPlaySession(DeviceId, DeviceName);
+	FTargetDeviceId TargetDeviceId;
+	if (FTargetDeviceId::Parse(DeviceId, TargetDeviceId))
+	{
+		if (FModuleManager::LoadModuleChecked<IProjectTargetPlatformEditorModule>("ProjectTargetPlatformEditor").ShowUnsupportedTargetWarning(*TargetDeviceId.GetPlatformName()))
+		{
+			GUnrealEd->RequestPlaySession(DeviceId, DeviceName);
+		}
+	}
 }
 
 
diff --git a/Engine/Source/Editor/UnrealEd/UnrealEd.Build.cs b/Engine/Source/Editor/UnrealEd/UnrealEd.Build.cs
--- a/Engine/Source/Editor/UnrealEd/UnrealEd.Build.cs
+++ b/Engine/Source/Editor/UnrealEd/UnrealEd.Build.cs
@@ -34,6 +34,7 @@ public class UnrealEd : ModuleRules
 				"DesktopPlatform",
                 "EnvironmentQueryEditor",
 				"GameProjectGeneration",
+				"ProjectTargetPlatformEditor",
 				"ImageWrapper",
 				"MainFrame",
 				"MaterialEditor",
@@ -159,6 +160,7 @@ public class UnrealEd : ModuleRules
 				"NiagaraEditor",
 				"MeshUtilities",
 				"GameProjectGeneration",
+				"ProjectTargetPlatformEditor",
 				"PListEditor",
                 "Documentation",
                 "BehaviorTreeEditor",
diff --git a/Engine/Source/Runtime/Projects/Private/ProjectManager.cpp b/Engine/Source/Runtime/Projects/Private/ProjectManager.cpp
--- a/Engine/Source/Runtime/Projects/Private/ProjectManager.cpp
+++ b/Engine/Source/Runtime/Projects/Private/ProjectManager.cpp
@@ -28,16 +28,50 @@ void FProject::SignSampleProject(const FString& FilePath, const FString& Categor
 	ProjectInfo.Category = Category;
 }
 
+void FProject::UpdateSupportedTargetPlatforms(const FName& InPlatformName, const bool bIsSupported)
+{
+	if ( bIsSupported )
+	{
+		ProjectInfo.TargetPlatforms.AddUnique(InPlatformName);
+	}
+	else
+	{
+		ProjectInfo.TargetPlatforms.Remove(InPlatformName);
+	}
+}
+
 bool FProject::PerformAdditionalDeserialization(const TSharedRef< FJsonObject >& FileObject)
 {
 	ReadNumberFromJSON(FileObject, TEXT("EpicSampleNameHash"), ProjectInfo.EpicSampleNameHash);
 
+	ProjectInfo.TargetPlatforms.Empty();
+	if ( FileObject->HasField(TEXT("TargetPlatforms")) )
+	{
+		const TSharedPtr<FJsonValue>& TargetPlatformsValue = FileObject->GetField<EJson::Array>(TEXT("TargetPlatforms"));
+		const TArray<TSharedPtr<FJsonValue>>& TargetPlatformsArray = TargetPlatformsValue->AsArray();
+		for ( const auto& TargetPlatformsEntry : TargetPlatformsArray )
+		{
+			const FString PlatformName = TargetPlatformsEntry->AsString();
+			ProjectInfo.TargetPlatforms.Add(*PlatformName);
+		}
+	}
+
 	return true;
 }
 
 void FProject::PerformAdditionalSerialization(const TSharedRef< TJsonWriter<> >& Writer) const
 {
 	Writer->WriteValue(TEXT("EpicSampleNameHash"), FString::Printf(TEXT("%u"), ProjectInfo.EpicSampleNameHash));
+
+	if ( ProjectInfo.TargetPlatforms.Num() > 0 )
+	{
+		Writer->WriteArrayStart(TEXT("TargetPlatforms"));
+		for ( const FName& PlatformName : ProjectInfo.TargetPlatforms )
+		{
+			Writer->WriteValue(PlatformName.ToString());
+		}
+		Writer->WriteArrayEnd();
+	}
 }
 
 
@@ -263,15 +297,69 @@ bool FProjectManager::QueryStatusForProject(const FString& FilePath, FProjectSta
 		return false;
 	}
 
-	const FProjectInfo& ProjectInfo = NewProject->GetProjectInfo();
+	QueryStatusForProjectImpl(*NewProject, FilePath, OutProjectStatus);
+	return true;
+}
+
+bool FProjectManager::QueryStatusForCurrentProject(FProjectStatus& OutProjectStatus) const
+{
+	if ( !CurrentlyLoadedProject.IsValid() )
+	{
+		return false;
+	}
+
+	QueryStatusForProjectImpl(*CurrentlyLoadedProject, FPaths::GetProjectFilePath(), OutProjectStatus);
+	return true;
+}
+
+void FProjectManager::QueryStatusForProjectImpl(const FProject& Project, const FString& FilePath, FProjectStatus& OutProjectStatus)
+{
+	const FProjectInfo& ProjectInfo = Project.GetProjectInfo();
 	OutProjectStatus.Name = ProjectInfo.Name;
 	OutProjectStatus.Description = ProjectInfo.Description;
 	OutProjectStatus.Category = ProjectInfo.Category;
 	OutProjectStatus.bCodeBasedProject = ProjectInfo.Modules.Num() > 0;
-	OutProjectStatus.bSignedSampleProject = NewProject->IsSignedSampleProject(FilePath);
-	OutProjectStatus.bRequiresUpdate = NewProject->RequiresUpdate();
+	OutProjectStatus.bSignedSampleProject = Project.IsSignedSampleProject(FilePath);
+	OutProjectStatus.bRequiresUpdate = Project.RequiresUpdate();
+	OutProjectStatus.TargetPlatforms = ProjectInfo.TargetPlatforms;
+}
 
-	return true;
+void FProjectManager::UpdateSupportedTargetPlatformsForProject(const FString& FilePath, const FName& InPlatformName, const bool bIsSupported)
+{
+	TSharedRef<FProject> NewProject = MakeShareable( new FProject() );
+	FText FailReason;
+	if ( !NewProject->LoadFromFile(FilePath, FailReason) )
+	{
+		return;
+	}
+
+	NewProject->UpdateSupportedTargetPlatforms(InPlatformName, bIsSupported);
+
+	const FString& FileContents = NewProject->SerializeToJSON();
+	FFileHelper::SaveStringToFile(FileContents, *FilePath);
+
+	// Call OnTargetPlatformsForCurrentProjectChangedEvent if this project is the same as the one we currently have loaded
+	const FString CurrentProjectPath = FPaths::ConvertRelativePathToFull(FPaths::GetProjectFilePath());
+	const FString InProjectPath = FPaths::ConvertRelativePathToFull(FilePath);
+	if ( CurrentProjectPath == InProjectPath )
+	{
+		OnTargetPlatformsForCurrentProjectChangedEvent.Broadcast();
+	}
+}
+
+void FProjectManager::UpdateSupportedTargetPlatformsForCurrentProject(const FName& InPlatformName, const bool bIsSupported)
+{
+	if ( !CurrentlyLoadedProject.IsValid() )
+	{
+		return;
+	}
+
+	CurrentlyLoadedProject->UpdateSupportedTargetPlatforms(InPlatformName, bIsSupported);
+
+	const FString& FileContents = CurrentlyLoadedProject->SerializeToJSON();
+	FFileHelper::SaveStringToFile(FileContents, *FPaths::GetProjectFilePath());
+
+	OnTargetPlatformsForCurrentProjectChangedEvent.Broadcast();
 }
 
 IProjectManager& IProjectManager::Get()
diff --git a/Engine/Source/Runtime/Projects/Private/ProjectManager.h b/Engine/Source/Runtime/Projects/Private/ProjectManager.h
--- a/Engine/Source/Runtime/Projects/Private/ProjectManager.h
+++ b/Engine/Source/Runtime/Projects/Private/ProjectManager.h
@@ -18,6 +18,8 @@ public:
 	/** A hash that is used to determine if the project was forked from a sample */
 	uint32 EpicSampleNameHash;
 
+	/** Array of platforms that this project is targeting */
+	TArray<FName> TargetPlatforms;
 
 	/** Project info constructor */
 	FProjectInfo()
@@ -55,6 +57,14 @@ public:
 	 */
 	void SignSampleProject(const FString& FilePath, const FString &Category);
 
+	/**
+	 * Update the list of supported target platforms based upon the parameters provided
+	 * 
+	 * @param	InPlatformName		Name of the platform to target (eg, WindowsNoEditor)
+	 * @param	bIsSupported		true if the platform should be supported by this project, false if it should not
+	 */
+	void UpdateSupportedTargetPlatforms(const FName& InPlatformName, const bool bIsSupported);
+
 	const FProjectInfo& GetProjectInfo() const
 	{
 		return ProjectInfo;
@@ -92,10 +102,19 @@ public:
 	virtual bool UpdateLoadedProjectFileToCurrent(const TArray<FString>* StartupModuleNames, const FString& EngineIdentifier, FText& OutFailReason) OVERRIDE;
 	virtual bool SignSampleProject(const FString& FilePath, const FString& Category, FText& OutFailReason) OVERRIDE;
 	virtual bool QueryStatusForProject(const FString& FilePath, FProjectStatus& OutProjectStatus) const OVERRIDE;
+	virtual bool QueryStatusForCurrentProject(FProjectStatus& OutProjectStatus) const OVERRIDE;
+	virtual void UpdateSupportedTargetPlatformsForProject(const FString& FilePath, const FName& InPlatformName, const bool bIsSupported) OVERRIDE;
+	virtual void UpdateSupportedTargetPlatformsForCurrentProject(const FName& InPlatformName, const bool bIsSupported) OVERRIDE;
+	virtual FOnTargetPlatformsForCurrentProjectChangedEvent& OnTargetPlatformsForCurrentProjectChanged() OVERRIDE { return OnTargetPlatformsForCurrentProjectChangedEvent; }
 
 private:
+	static void QueryStatusForProjectImpl(const FProject& Project, const FString& FilePath, FProjectStatus& OutProjectStatus);
+
 	/** The project that is currently loaded in the editor */
 	TSharedPtr< FProject > CurrentlyLoadedProject;
+
+	/** Delegate called when the target platforms for the current project are changed */
+	FOnTargetPlatformsForCurrentProjectChangedEvent OnTargetPlatformsForCurrentProjectChangedEvent;
 };
 
 
diff --git a/Engine/Source/Runtime/Projects/Public/Interfaces/IProjectManager.h b/Engine/Source/Runtime/Projects/Public/Interfaces/IProjectManager.h
--- a/Engine/Source/Runtime/Projects/Public/Interfaces/IProjectManager.h
+++ b/Engine/Source/Runtime/Projects/Public/Interfaces/IProjectManager.h
@@ -31,11 +31,26 @@ public:
 	/** True if this project needs to be updated */
 	bool bRequiresUpdate;
 
+	/** Array of platforms that this project is targeting */
+	TArray<FName> TargetPlatforms;
+
 	FProjectStatus()
 		: bSignedSampleProject(false)
 		, bCodeBasedProject(false)
 		, bRequiresUpdate(false)
 	{}
+
+	/**
+	 * Check to see if the given platform name is supported as a target by the current project
+	 *
+	 * @param	InPlatformName				Name of the platform to target (eg, WindowsNoEditor)
+	 * @param	bAllowSupportedIfEmptyList	Consider an empty list to mean that all platforms are supported targets?
+	 */
+	bool IsTargetPlatformSupported(const FName& InPlatformName, const bool bAllowSupportedIfEmptyList = true) const
+	{
+		// An empty list is considered the same as supporting all platforms
+		return (bAllowSupportedIfEmptyList && TargetPlatforms.Num() == 0) || TargetPlatforms.Contains(InPlatformName);
+	}
 };
 
 /**
@@ -147,13 +162,42 @@ public:
 	 * Gets status about the specified project
 	 *
 	 * @param FilePath				The filepath where the project is stored.
-	 * @param EngineIdentifier		Identifier for the current engine installation
 	 * @param OutProjectStatus		The status for the project.
 	 *
 	 * @return	 true if the file was successfully open and read
 	 */
 	virtual bool QueryStatusForProject(const FString& FilePath, FProjectStatus& OutProjectStatus) const = 0;
 
+	/**
+	 * Gets status about the current project
+	 *
+	 * @param OutProjectStatus		The status for the project.
+	 *
+	 * @return	 true if the file was successfully open and read
+	 */
+	virtual bool QueryStatusForCurrentProject(FProjectStatus& OutProjectStatus) const = 0;
+
+	/**
+	 * Update the list of supported target platforms for the target project based upon the parameters provided
+	 * 
+	 * @param	FilePath			The filepath where the project is stored.
+	 * @param	InPlatformName		Name of the platform to target (eg, WindowsNoEditor)
+	 * @param	bIsSupported		true if the platform should be supported by this project, false if it should not
+	 */
+	virtual void UpdateSupportedTargetPlatformsForProject(const FString& FilePath, const FName& InPlatformName, const bool bIsSupported) = 0;
+
+	/**
+	 * Update the list of supported target platforms for the current project based upon the parameters provided
+	 * 
+	 * @param	InPlatformName		Name of the platform to target (eg, WindowsNoEditor)
+	 * @param	bIsSupported		true if the platform should be supported by this project, false if it should not
+	 */
+	virtual void UpdateSupportedTargetPlatformsForCurrentProject(const FName& InPlatformName, const bool bIsSupported) = 0;
+
+	/** Called when the target platforms for the current project are changed */
+	DECLARE_MULTICAST_DELEGATE(FOnTargetPlatformsForCurrentProjectChangedEvent);
+	virtual FOnTargetPlatformsForCurrentProjectChangedEvent& OnTargetPlatformsForCurrentProjectChanged() = 0;
+
 	/** Helper functions to reduce the syntax complexity of commonly used functions */
 	static const FString& GetProjectFileExtension() { return Get().NonStaticGetProjectFileExtension(); }
 };
\ No newline at end of file
diff --git a/Engine/Source/Runtime/Slate/Private/Framework/MultiBox/MultiBoxBuilder.cpp b/Engine/Source/Runtime/Slate/Private/Framework/MultiBox/MultiBoxBuilder.cpp
--- a/Engine/Source/Runtime/Slate/Private/Framework/MultiBox/MultiBoxBuilder.cpp
+++ b/Engine/Source/Runtime/Slate/Private/Framework/MultiBox/MultiBoxBuilder.cpp
@@ -119,13 +119,13 @@ void FBaseMenuBuilder::AddMenuEntry( const TAttribute<FText>& InLabel, const TAt
 	ApplyHook(InExtensionHook, EExtensionHook::After);
 }
 
-void FBaseMenuBuilder::AddMenuEntry( const FUIAction& UIAction, const TSharedRef< SWidget > Contents, const FName& InExtensionHook, const EUserInterfaceActionType::Type UserInterfaceActionType )
+void FBaseMenuBuilder::AddMenuEntry( const FUIAction& UIAction, const TSharedRef< SWidget > Contents, const FName& InExtensionHook, const TAttribute<FText>& InToolTip, const EUserInterfaceActionType::Type UserInterfaceActionType )
 {
 	ApplySectionBeginning();
 
 	ApplyHook(InExtensionHook, EExtensionHook::Before);
 
-	TSharedRef< FMenuEntryBlock > NewMenuEntryBlock( new FMenuEntryBlock( InExtensionHook, UIAction, Contents, UserInterfaceActionType, bCloseSelfOnly ) );
+	TSharedRef< FMenuEntryBlock > NewMenuEntryBlock( new FMenuEntryBlock( InExtensionHook, UIAction, Contents, InToolTip, UserInterfaceActionType, bCloseSelfOnly ) );
 	MultiBox->AddMultiBlock( NewMenuEntryBlock );
 
 	ApplyHook(InExtensionHook, EExtensionHook::After);
diff --git a/Engine/Source/Runtime/Slate/Private/Framework/MultiBox/SMenuEntryBlock.cpp b/Engine/Source/Runtime/Slate/Private/Framework/MultiBox/SMenuEntryBlock.cpp
--- a/Engine/Source/Runtime/Slate/Private/Framework/MultiBox/SMenuEntryBlock.cpp
+++ b/Engine/Source/Runtime/Slate/Private/Framework/MultiBox/SMenuEntryBlock.cpp
@@ -90,8 +90,9 @@ FMenuEntryBlock::FMenuEntryBlock( const FName& InExtensionHook, const TAttribute
 {
 }
 
-FMenuEntryBlock::FMenuEntryBlock( const FName& InExtensionHook, const FUIAction& UIAction, const TSharedRef< SWidget > Contents, const EUserInterfaceActionType::Type InUserInterfaceActionType, bool bInCloseSelfOnly )
+FMenuEntryBlock::FMenuEntryBlock( const FName& InExtensionHook, const FUIAction& UIAction, const TSharedRef< SWidget > Contents, const TAttribute<FText>& InToolTip, const EUserInterfaceActionType::Type InUserInterfaceActionType, bool bInCloseSelfOnly )
 	: FMultiBlock( UIAction, InExtensionHook )
+	, ToolTipOverride( InToolTip )
 	, EntryWidget( Contents )
 	, bIsSubMenu( false )
 	, bOpenSubMenuOnClick( false )
@@ -115,7 +116,7 @@ FMenuEntryBlock::FMenuEntryBlock( const FName& InExtensionHook, const TSharedRef
 
 
 FMenuEntryBlock::FMenuEntryBlock( const FName& InExtensionHook, const FUIAction& UIAction, const TSharedRef< SWidget > Contents, const FNewMenuDelegate& InEntryBuilder, TSharedPtr<FExtender> InExtender, bool bInSubMenu, TSharedPtr< const FUICommandList > InCommandList, bool bInCloseSelfOnly )
-: FMultiBlock( UIAction, InExtensionHook )
+	: FMultiBlock( UIAction, InExtensionHook )
 	, EntryBuilder( InEntryBuilder )
 	, EntryWidget( Contents )
 	, bIsSubMenu( bInSubMenu )
@@ -321,6 +322,12 @@ TSharedRef< SWidget > SMenuEntryBlock::BuildMenuEntryWidget( const FMenuEntryBui
 				EVisibility::Visible :
 				EVisibility::Hidden;
 
+	// Collapse (rather than Hide) the checkbox when using a custom menu widget with a button action, otherwise we add additional padding around the user defined widget
+	if (MenuEntryBlock->EntryWidget.IsValid() && UserInterfaceType == EUserInterfaceActionType::Button)
+	{
+		CheckBoxVisibility = EVisibility::Collapsed;
+	}
+
 	TAttribute<FSlateColor> CheckBoxForegroundColor = FSlateColor::UseForeground();
 	FName CheckBoxStyle = ISlateStyle::Join( StyleName, ".CheckBox" );
 	if (UserInterfaceType == EUserInterfaceActionType::Check)
@@ -341,24 +348,6 @@ TSharedRef< SWidget > SMenuEntryBlock::BuildMenuEntryWidget( const FMenuEntryBui
 		// Whatever we have in the icon area goes first
 		+ SHorizontalBox::Slot()
 		.AutoWidth()
-		[
-			SNew(SBox)
-			.Visibility(CheckBoxVisibility)
-			.WidthOverride( MultiBoxConstants::MenuCheckBoxSize )
-			.HeightOverride( MultiBoxConstants::MenuCheckBoxSize )
-			.HAlign(HAlign_Center)
-			.VAlign(VAlign_Center)
-			[
-				SNew( SCheckBox )
-				.ForegroundColor( CheckBoxForegroundColor )
-				.IsChecked( this, &SMenuEntryBlock::IsChecked )
-				.Style( StyleSet, CheckBoxStyle )
-				.OnCheckStateChanged( this, &SMenuEntryBlock::OnCheckStateChanged )
-				.ReadOnly( UserInterfaceType == EUserInterfaceActionType::Check )
-			]
-		]
-		+ SHorizontalBox::Slot()
-		.AutoWidth()
 		.Padding(FMargin(2, 0, 2, 0))
 		[
 			SNew( SBox )
@@ -402,6 +391,33 @@ TSharedRef< SWidget > SMenuEntryBlock::BuildMenuEntryWidget( const FMenuEntryBui
 		];
 	}
 
+	// Create a wrapper containing the checkbox, and the generated button content
+	TSharedPtr< SWidget > CheckBoxAndButtonContent = 
+		SNew( SHorizontalBox )
+		+SHorizontalBox::Slot()
+		.AutoWidth()
+		[
+			SNew(SBox)
+			.Visibility(CheckBoxVisibility)
+			.WidthOverride( MultiBoxConstants::MenuCheckBoxSize )
+			.HeightOverride( MultiBoxConstants::MenuCheckBoxSize )
+			.HAlign(HAlign_Center)
+			.VAlign(VAlign_Center)
+			[
+				SNew( SCheckBox )
+				.ForegroundColor( CheckBoxForegroundColor )
+				.IsChecked( this, &SMenuEntryBlock::IsChecked )
+				.Style( StyleSet, CheckBoxStyle )
+				.OnCheckStateChanged( this, &SMenuEntryBlock::OnCheckStateChanged )
+				.ReadOnly( UserInterfaceType == EUserInterfaceActionType::Check )
+			]
+		]
+		+SHorizontalBox::Slot()
+		.AutoWidth()
+		[
+			ButtonContent.ToSharedRef()
+		];
+
 	// Create a menu item button
 	TSharedPtr<SWidget> MenuEntryWidget = 
 		SNew( SButton )
@@ -419,7 +435,7 @@ TSharedRef< SWidget > SMenuEntryBlock::BuildMenuEntryWidget( const FMenuEntryBui
 
 		.ForegroundColor( FSlateColor::UseForeground() )
 		[
-			ButtonContent.ToSharedRef()
+			CheckBoxAndButtonContent.ToSharedRef()
 		]
 
 		// Bind the button's "on clicked" event to our object's method for this
diff --git a/Engine/Source/Runtime/Slate/Private/Framework/MultiBox/SMenuEntryBlock.h b/Engine/Source/Runtime/Slate/Private/Framework/MultiBox/SMenuEntryBlock.h
--- a/Engine/Source/Runtime/Slate/Private/Framework/MultiBox/SMenuEntryBlock.h
+++ b/Engine/Source/Runtime/Slate/Private/Framework/MultiBox/SMenuEntryBlock.h
@@ -80,7 +80,7 @@ public:
 
 	FMenuEntryBlock( const FName& InExtensionHook, const TAttribute<FText>& InLabel, const TAttribute<FText>& InToolTip, const FSlateIcon& InIcon, const FUIAction& InUIAction, const EUserInterfaceActionType::Type InUserInterfaceActionType, bool bInCloseSelfOnly );
 
-	FMenuEntryBlock( const FName& InExtensionHook, const FUIAction& UIAction, const TSharedRef< SWidget > Contents, const EUserInterfaceActionType::Type InUserInterfaceActionType, bool bInCloseSelfOnly );
+	FMenuEntryBlock( const FName& InExtensionHook, const FUIAction& UIAction, const TSharedRef< SWidget > Contents, const TAttribute<FText>& InToolTip, const EUserInterfaceActionType::Type InUserInterfaceActionType, bool bInCloseSelfOnly );
 
 	FMenuEntryBlock( const FName& InExtensionHook, const TSharedRef< SWidget > Contents, const FNewMenuDelegate& InEntryBuilder, TSharedPtr<FExtender> InExtender, bool bInSubMenu, bool bInSubMenuOnClick, TSharedPtr< const FUICommandList > InCommandList, bool bInCloseSelfOnly );
 
diff --git a/Engine/Source/Runtime/Slate/Public/Framework/MultiBox/MultiBoxBuilder.h b/Engine/Source/Runtime/Slate/Public/Framework/MultiBox/MultiBoxBuilder.h
--- a/Engine/Source/Runtime/Slate/Public/Framework/MultiBox/MultiBoxBuilder.h
+++ b/Engine/Source/Runtime/Slate/Public/Framework/MultiBox/MultiBoxBuilder.h
@@ -157,7 +157,7 @@ public:
 	 */
 	void AddMenuEntry( const TAttribute<FText>& InLabel, const TAttribute<FText>& InToolTip, const FSlateIcon& InIcon, const FUIAction& UIAction, FName InExtensionHook = NAME_None, const EUserInterfaceActionType::Type UserInterfaceActionType = EUserInterfaceActionType::Button );
 
-	void AddMenuEntry( const FUIAction& UIAction, const TSharedRef< SWidget > Contents, const FName& InExtensionHook = NAME_None, const EUserInterfaceActionType::Type UserInterfaceActionType = EUserInterfaceActionType::Button );
+	void AddMenuEntry( const FUIAction& UIAction, const TSharedRef< SWidget > Contents, const FName& InExtensionHook = NAME_None, const TAttribute<FText>& InToolTip = TAttribute<FText>(), const EUserInterfaceActionType::Type UserInterfaceActionType = EUserInterfaceActionType::Button );
 
 protected:
 	/** True if clicking on a menu entry closes itself only and its children and not the entire stack */
