Renamed "Target Platform" to "Supported Platform" in the UI

Changed anything using "target" to use "support" instead. Made some description text less verbose.

diff --git a/Engine/Source/Editor/ProjectSettingsViewer/Private/ProjectSettingsViewerModule.cpp b/Engine/Source/Editor/ProjectSettingsViewer/Private/ProjectSettingsViewerModule.cpp
--- a/Engine/Source/Editor/ProjectSettingsViewer/Private/ProjectSettingsViewerModule.cpp
+++ b/Engine/Source/Editor/ProjectSettingsViewer/Private/ProjectSettingsViewerModule.cpp
@@ -182,9 +182,9 @@ protected:
 
 		// platforms settings
 		TWeakPtr<SWidget> ProjectTargetPlatformEditorPanel = FModuleManager::LoadModuleChecked<IProjectTargetPlatformEditorModule>("ProjectTargetPlatformEditor").CreateProjectTargetPlatformEditorPanel();
-		SettingsModule.RegisterSettings("Project", "Game", "TargetPlatforms",
-			LOCTEXT("ProjectTargetPlatformsSettingsName", "Target Platforms"),
-			LOCTEXT("ProjectTargetPlatformsSettingsDescription", "Specify which platforms your project is targeting."),
+		SettingsModule.RegisterSettings("Project", "Game", "SupportedPlatforms",
+			LOCTEXT("ProjectSupportedPlatformsSettingsName", "Supported Platforms"),
+			LOCTEXT("ProjectSupportedPlatformsSettingsDescription", "Specify which platforms your project supports."),
 			ProjectTargetPlatformEditorPanel.Pin().ToSharedRef()
 		);
 
@@ -235,7 +235,7 @@ protected:
 			SettingsModule->UnregisterSettings("Project", "Game", "General");
 			SettingsModule->UnregisterSettings("Project", "Game", "Maps");
 			SettingsModule->UnregisterSettings("Project", "Game", "Packaging");
-			SettingsModule->UnregisterSettings("Project", "Game", "TargetPlatforms");
+			SettingsModule->UnregisterSettings("Project", "Game", "SupportedPlatforms");
 			SettingsModule->UnregisterSettings("Project", "Game", "Movies");
 //			SettingsModule->UnregisterSettings("Project", "Game", "GameSession");
 //			SettingsModule->UnregisterSettings("Project", "Game", "HUD");
diff --git a/Engine/Source/Editor/ProjectTargetPlatformEditor/Private/ProjectTargetPlatformEditorModule.cpp b/Engine/Source/Editor/ProjectTargetPlatformEditor/Private/ProjectTargetPlatformEditorModule.cpp
--- a/Engine/Source/Editor/ProjectTargetPlatformEditor/Private/ProjectTargetPlatformEditorModule.cpp
+++ b/Engine/Source/Editor/ProjectTargetPlatformEditor/Private/ProjectTargetPlatformEditorModule.cpp
@@ -38,10 +38,10 @@ public:
 		};
 
 		MenuBuilder.AddMenuEntry(
-			LOCTEXT("TargetPlatformsMenuLabel", "Target Platforms..."),
-			LOCTEXT("ProjectSettingsMenuToolTip", "Change which platforms this project is targeting"),
+			LOCTEXT("SupportedPlatformsMenuLabel", "Supported Platforms..."),
+			LOCTEXT("SupportedPlatformsMenuToolTip", "Change which platforms this project supports"),
 			FSlateIcon(),
-			FUIAction(FExecuteAction::CreateStatic(&Local::OpenSettings, FName("Project"), FName("Game"), FName("TargetPlatforms")))
+			FUIAction(FExecuteAction::CreateStatic(&Local::OpenSettings, FName("Project"), FName("Game"), FName("SupportedPlatforms")))
 			);
 	}
 
@@ -114,18 +114,18 @@ public:
 
 			FFormatNamedArguments Args;
 			Args.Add(TEXT("DisplayName"), Platform->DisplayName());
-			FText WarningText = FText::Format(LOCTEXT("ShowUnsupportedTargetWarning_Message", "{DisplayName} is not listed as a target platform for this project, so may not run as expected.\n\nDo you wish to continue?"), Args);
+			FText WarningText = FText::Format(LOCTEXT("ShowUnsupportedPlatformWarning_Message", "{DisplayName} is not listed as a supported platform for this project, so may not run as expected.\n\nDo you wish to continue?"), Args);
 
 			FSuppressableWarningDialog::FSetupInfo Info(
 				WarningText, 
-				LOCTEXT("ShowUnsupportedTargetWarning_Title", "Unsupported Platform"), 
-				TEXT("SuppressUnsupportedTargetWarningDialog")
+				LOCTEXT("ShowUnsupportedPlatformWarning_Title", "Unsupported Platform"), 
+				TEXT("SuppressUnsupportedPlatformWarningDialog")
 				);
-			Info.ConfirmText = LOCTEXT("ShowUnsupportedTargetWarning_Confirm", "Continue");
-			Info.CancelText = LOCTEXT("ShowUnsupportedTargetWarning_Cancel", "Cancel");
-			FSuppressableWarningDialog UnsupportedTargetWarningDialog(Info);
+			Info.ConfirmText = LOCTEXT("ShowUnsupportedPlatformWarning_Confirm", "Continue");
+			Info.CancelText = LOCTEXT("ShowUnsupportedPlatformWarning_Cancel", "Cancel");
+			FSuppressableWarningDialog UnsupportedPlatformWarningDialog(Info);
 
-			return UnsupportedTargetWarningDialog.ShowModal() != FSuppressableWarningDialog::EResult::Cancel;
+			return UnsupportedPlatformWarningDialog.ShowModal() != FSuppressableWarningDialog::EResult::Cancel;
 		}
 
 		return true;
diff --git a/Engine/Source/Editor/ProjectTargetPlatformEditor/Private/Widgets/SProjectTargetPlatformSettings.cpp b/Engine/Source/Editor/ProjectTargetPlatformEditor/Private/Widgets/SProjectTargetPlatformSettings.cpp
--- a/Engine/Source/Editor/ProjectTargetPlatformEditor/Private/Widgets/SProjectTargetPlatformSettings.cpp
+++ b/Engine/Source/Editor/ProjectTargetPlatformEditor/Private/Widgets/SProjectTargetPlatformSettings.cpp
@@ -175,7 +175,7 @@ void SProjectTargetPlatformSettings::Construct(const FArguments& InArgs)
 				[
 					SNew(STextBlock)
 					.AutoWrapText(true)
-					.Text(LOCTEXT("PlatformsListDescription", "Use the list above to select the platforms that your project is targeting. Attempting to run your project on a platform that hasn't been selected will result in a warning stating that the project may not run as expected."))
+					.Text(LOCTEXT("PlatformsListDescription", "Select the supported platforms for your project. Attempting to package, run, or cook your project on an unsupported platform will result in a warning."))
 				]
 
 				+SVerticalBox::Slot()
@@ -183,7 +183,7 @@ void SProjectTargetPlatformSettings::Construct(const FArguments& InArgs)
 				[
 					SNew(STextBlock)
 					.AutoWrapText(true)
-					.Text(LOCTEXT("PlatformsListNote", "Please note that selecting no platforms is considered the same as selecting all platforms."))
+					.Text(LOCTEXT("PlatformsListNote", "Note: Selecting no platforms is considered the same as selecting all platforms."))
 				]
 			]
 		]
