Remove project/plugin includes from UnrealEd. Only a handful of things needs to access it, so just include it manually.

diff --git a/Engine/Plugins/Editor/PluginsEditor/Source/PluginsEditor/Private/SPluginListTile.cpp b/Engine/Plugins/Editor/PluginsEditor/Source/PluginsEditor/Private/SPluginListTile.cpp
--- a/Engine/Plugins/Editor/PluginsEditor/Source/PluginsEditor/Private/SPluginListTile.cpp
+++ b/Engine/Plugins/Editor/PluginsEditor/Source/PluginsEditor/Private/SPluginListTile.cpp
@@ -8,6 +8,7 @@
 #include "PluginStyle.h"
 #include "GameProjectGenerationModule.h"
 #include "SHyperlink.h"
+#include "IProjectManager.h"
 
 #define LOCTEXT_NAMESPACE "PluginListTile"
 
diff --git a/Engine/Source/Editor/ContentBrowser/Private/ContentBrowserUtils.cpp b/Engine/Source/Editor/ContentBrowser/Private/ContentBrowserUtils.cpp
--- a/Engine/Source/Editor/ContentBrowser/Private/ContentBrowserUtils.cpp
+++ b/Engine/Source/Editor/ContentBrowser/Private/ContentBrowserUtils.cpp
@@ -12,6 +12,7 @@
 #include "EngineBuildSettings.h"
 #include "SNotificationList.h"
 #include "NotificationManager.h"
+#include "IPluginManager.h"
 
 #define LOCTEXT_NAMESPACE "ContentBrowser"
 
diff --git a/Engine/Source/Editor/Kismet/Private/BlueprintEditor.cpp b/Engine/Source/Editor/Kismet/Private/BlueprintEditor.cpp
--- a/Engine/Source/Editor/Kismet/Private/BlueprintEditor.cpp
+++ b/Engine/Source/Editor/Kismet/Private/BlueprintEditor.cpp
@@ -42,6 +42,8 @@
 
 #include "ClassIconFinder.h"
 
+#include "IProjectManager.h"
+
 // Core kismet tabs
 #include "SGraphNode.h"
 #include "SSCSEditor.h"
diff --git a/Engine/Source/Editor/LevelEditor/Private/LevelEditor.cpp b/Engine/Source/Editor/LevelEditor/Private/LevelEditor.cpp
--- a/Engine/Source/Editor/LevelEditor/Private/LevelEditor.cpp
+++ b/Engine/Source/Editor/LevelEditor/Private/LevelEditor.cpp
@@ -20,6 +20,7 @@
 #include "ToolkitManager.h"
 #include "TargetPlatform.h"
 #include "IIntroTutorials.h"
+#include "IProjectManager.h"
 
 // @todo Editor: remove this circular dependency
 #include "Editor/MainFrame/Public/Interfaces/IMainFrameModule.h"
diff --git a/Engine/Source/Editor/MainFrame/Private/Menus/CookContentMenu.h b/Engine/Source/Editor/MainFrame/Private/Menus/CookContentMenu.h
--- a/Engine/Source/Editor/MainFrame/Private/Menus/CookContentMenu.h
+++ b/Engine/Source/Editor/MainFrame/Private/Menus/CookContentMenu.h
@@ -3,6 +3,7 @@
 #pragma once
 
 #include "ProjectTargetPlatformEditor.h"
+#include "IProjectManager.h"
 
 
 #define LOCTEXT_NAMESPACE "FCookContentMenu"
diff --git a/Engine/Source/Editor/ProjectTargetPlatformEditor/Private/ProjectTargetPlatformEditorModule.cpp b/Engine/Source/Editor/ProjectTargetPlatformEditor/Private/ProjectTargetPlatformEditorModule.cpp
--- a/Engine/Source/Editor/ProjectTargetPlatformEditor/Private/ProjectTargetPlatformEditorModule.cpp
+++ b/Engine/Source/Editor/ProjectTargetPlatformEditor/Private/ProjectTargetPlatformEditorModule.cpp
@@ -2,6 +2,7 @@
 
 #include "ProjectTargetPlatformEditorPrivatePCH.h"
 #include "ISettingsModule.h"
+#include "IProjectManager.h"
 
 
 #define LOCTEXT_NAMESPACE "FProjectTargetPlatformEditorModule"
diff --git a/Engine/Source/Editor/UnrealEd/Private/Commandlets/UpdateGameProjectCommandlet.cpp b/Engine/Source/Editor/UnrealEd/Private/Commandlets/UpdateGameProjectCommandlet.cpp
--- a/Engine/Source/Editor/UnrealEd/Private/Commandlets/UpdateGameProjectCommandlet.cpp
+++ b/Engine/Source/Editor/UnrealEd/Private/Commandlets/UpdateGameProjectCommandlet.cpp
@@ -4,6 +4,7 @@
 #include "UnrealEd.h"
 #include "ISourceControlModule.h"
 #include "GameProjectGenerationModule.h"
+#include "IProjectManager.h"
 
 DEFINE_LOG_CATEGORY_STATIC(LogUpdateGameProjectCommandlet, Log, All);
 
diff --git a/Engine/Source/Editor/UnrealEd/Private/Kismet2/DebuggerCommands.cpp b/Engine/Source/Editor/UnrealEd/Private/Kismet2/DebuggerCommands.cpp
--- a/Engine/Source/Editor/UnrealEd/Private/Kismet2/DebuggerCommands.cpp
+++ b/Engine/Source/Editor/UnrealEd/Private/Kismet2/DebuggerCommands.cpp
@@ -31,6 +31,8 @@
 #include "GameFramework/PlayerStart.h"
 #include "Components/CapsuleComponent.h"
 
+#include "IProjectManager.h"
+
 
 #define LOCTEXT_NAMESPACE "DebuggerCommands"
 
diff --git a/Engine/Source/Editor/UnrealEd/Private/Settings/EditorSettings.cpp b/Engine/Source/Editor/UnrealEd/Private/Settings/EditorSettings.cpp
--- a/Engine/Source/Editor/UnrealEd/Private/Settings/EditorSettings.cpp
+++ b/Engine/Source/Editor/UnrealEd/Private/Settings/EditorSettings.cpp
@@ -2,6 +2,7 @@
 
 #include "UnrealEd.h"
 #include "EditorSettings.h"
+#include "IProjectManager.h"
 
 UEditorSettings::UEditorSettings(const FObjectInitializer& ObjectInitializer)
 	: Super(ObjectInitializer)
diff --git a/Engine/Source/Editor/UnrealEd/Private/UnrealEdMisc.cpp b/Engine/Source/Editor/UnrealEd/Private/UnrealEdMisc.cpp
--- a/Engine/Source/Editor/UnrealEd/Private/UnrealEdMisc.cpp
+++ b/Engine/Source/Editor/UnrealEd/Private/UnrealEdMisc.cpp
@@ -41,6 +41,7 @@
 #include "GameMapsSettings.h"
 #include "GeneralProjectSettings.h"
 #include "Lightmass/LightmappedSurfaceCollection.h"
+#include "IProjectManager.h"
 
 #define LOCTEXT_NAMESPACE "UnrealEd"
 
diff --git a/Engine/Source/Editor/UnrealEd/Public/UnrealEd.h b/Engine/Source/Editor/UnrealEd/Public/UnrealEd.h
--- a/Engine/Source/Editor/UnrealEd/Public/UnrealEd.h
+++ b/Engine/Source/Editor/UnrealEd/Public/UnrealEd.h
@@ -29,7 +29,6 @@
 #include "EditorModes.h"
 
 #include "MRUList.h"
-#include "Projects.h"
 
 
 //#include "../Private/GeomFitUtils.h"
