Fixed unity compile errors

diff --git a/Engine/Plugins/Editor/PluginBrowser/Source/PluginBrowser/Private/SPluginBrowser.cpp b/Engine/Plugins/Editor/PluginBrowser/Source/PluginBrowser/Private/SPluginBrowser.cpp
--- a/Engine/Plugins/Editor/PluginBrowser/Source/PluginBrowser/Private/SPluginBrowser.cpp
+++ b/Engine/Plugins/Editor/PluginBrowser/Source/PluginBrowser/Private/SPluginBrowser.cpp
@@ -6,6 +6,7 @@
 #include "SPluginTileList.h"
 #include "SPluginCategoryTree.h"
 #include "SSearchBox.h"
+#include "PluginBrowserModule.h"
 
 #define LOCTEXT_NAMESPACE "PluginsEditor"
 
diff --git a/Engine/Plugins/Editor/PluginBrowser/Source/PluginBrowser/Private/SPluginTile.cpp b/Engine/Plugins/Editor/PluginBrowser/Source/PluginBrowser/Private/SPluginTile.cpp
--- a/Engine/Plugins/Editor/PluginBrowser/Source/PluginBrowser/Private/SPluginTile.cpp
+++ b/Engine/Plugins/Editor/PluginBrowser/Source/PluginBrowser/Private/SPluginTile.cpp
@@ -12,6 +12,7 @@
 #include "IProjectManager.h"
 #include "PluginBrowserModule.h"
 #include "ISourceControlModule.h"
+#include "PropertyEditorModule.h"
 
 #define LOCTEXT_NAMESPACE "PluginListTile"
 
diff --git a/Engine/Plugins/Editor/PluginBrowser/Source/PluginBrowser/Private/SPluginTile.h b/Engine/Plugins/Editor/PluginBrowser/Source/PluginBrowser/Private/SPluginTile.h
--- a/Engine/Plugins/Editor/PluginBrowser/Source/PluginBrowser/Private/SPluginTile.h
+++ b/Engine/Plugins/Editor/PluginBrowser/Source/PluginBrowser/Private/SPluginTile.h
@@ -2,6 +2,8 @@
 
 #pragma once
 
+class UPluginMetadataObject;
+
 /**
  * Widget that represents a "tile" for a single plugin in our plugins list
  */
diff --git a/Engine/Source/Editor/ContentBrowser/Private/NativeClassHierarchy.cpp b/Engine/Source/Editor/ContentBrowser/Private/NativeClassHierarchy.cpp
--- a/Engine/Source/Editor/ContentBrowser/Private/NativeClassHierarchy.cpp
+++ b/Engine/Source/Editor/ContentBrowser/Private/NativeClassHierarchy.cpp
@@ -7,6 +7,7 @@
 #include "IProjectManager.h"
 #include "HotReloadInterface.h"
 #include "SourceCodeNavigation.h"
+#include "IPluginManager.h"
 
 TSharedRef<FNativeClassHierarchyNode> FNativeClassHierarchyNode::MakeFolderEntry(FName InEntryName, FString InEntryPath)
 {
diff --git a/Engine/Source/Editor/ProjectTargetPlatformEditor/Private/Widgets/SProjectTargetPlatformSettings.cpp b/Engine/Source/Editor/ProjectTargetPlatformEditor/Private/Widgets/SProjectTargetPlatformSettings.cpp
--- a/Engine/Source/Editor/ProjectTargetPlatformEditor/Private/Widgets/SProjectTargetPlatformSettings.cpp
+++ b/Engine/Source/Editor/ProjectTargetPlatformEditor/Private/Widgets/SProjectTargetPlatformSettings.cpp
@@ -1,6 +1,7 @@
 // Copyright 1998-2015 Epic Games, Inc. All Rights Reserved.
 
 #include "ProjectTargetPlatformEditorPrivatePCH.h"
+#include "IProjectManager.h"
 
 #define LOCTEXT_NAMESPACE "SProjectTargetPlatformSettings"
 
diff --git a/Engine/Source/Editor/Sequencer/Private/SSequencerCurveEditor.cpp b/Engine/Source/Editor/Sequencer/Private/SSequencerCurveEditor.cpp
--- a/Engine/Source/Editor/Sequencer/Private/SSequencerCurveEditor.cpp
+++ b/Engine/Source/Editor/Sequencer/Private/SSequencerCurveEditor.cpp
@@ -2,6 +2,7 @@
 
 #include "SequencerPrivatePCH.h"
 #include "SCurveEditor.h"
+#include "Sequencer.h"
 
 #define LOCTEXT_NAMESPACE "SequencerCurveEditor"
 
diff --git a/Engine/Source/Editor/Sequencer/Private/SequencerCurveOwner.cpp b/Engine/Source/Editor/Sequencer/Private/SequencerCurveOwner.cpp
--- a/Engine/Source/Editor/Sequencer/Private/SequencerCurveOwner.cpp
+++ b/Engine/Source/Editor/Sequencer/Private/SequencerCurveOwner.cpp
@@ -4,6 +4,7 @@
 #include "SequencerCurveOwner.h"
 #include "IKeyArea.h"
 #include "MovieSceneSection.h"
+#include "Sequencer.h"
 
 void GetAllKeyAreaNodes( TSharedPtr<FSequencerNodeTree> InSequencerNodeTree, TArray<TSharedRef<FSectionKeyAreaNode>>& KeyAreaNodes )
 {
diff --git a/Engine/Source/Editor/Sequencer/Private/SequencerSelection.cpp b/Engine/Source/Editor/Sequencer/Private/SequencerSelection.cpp
--- a/Engine/Source/Editor/Sequencer/Private/SequencerSelection.cpp
+++ b/Engine/Source/Editor/Sequencer/Private/SequencerSelection.cpp
@@ -2,6 +2,7 @@
 
 #include "SequencerPrivatePCH.h"
 #include "SequencerSelection.h"
+#include "MovieSceneSection.h"
 
 FSequencerSelection::FSequencerSelection()
 {
